Fix for curl on platforms that have a bespoke version of libCurl

#rnx

diff --git a/Engine/Source/Developer/Zen/Internal/ZenServerHttp.cpp b/Engine/Source/Developer/Zen/Internal/ZenServerHttp.cpp
--- a/Engine/Source/Developer/Zen/Internal/ZenServerHttp.cpp
+++ b/Engine/Source/Developer/Zen/Internal/ZenServerHttp.cpp
@@ -12,7 +12,9 @@
 #	include "Windows/AllowWindowsPlatformTypes.h"
 #endif
 
-#include "curl/curl.h"
+#if !defined(PLATFORM_CURL_INCLUDE)
+	#include "curl/curl.h"
+#endif
 
 #if PLATFORM_WINDOWS || PLATFORM_HOLOLENS
 #	include "Windows/HideWindowsPlatformTypes.h"
diff --git a/Engine/Source/Developer/Zen/Internal/ZenServerHttp.h b/Engine/Source/Developer/Zen/Internal/ZenServerHttp.h
--- a/Engine/Source/Developer/Zen/Internal/ZenServerHttp.h
+++ b/Engine/Source/Developer/Zen/Internal/ZenServerHttp.h
@@ -8,6 +8,22 @@
 #include "Containers/UnrealString.h"
 #include "Memory/MemoryFwd.h"
 
+// if there is a platform-specific include then it must be used in the header file in case it defines CURL_STRICTER
+#if defined(PLATFORM_CURL_INCLUDE)
+
+#if PLATFORM_MICROSOFT
+#include "Microsoft/AllowMicrosoftPlatformTypes.h"
+#endif
+
+#include PLATFORM_CURL_INCLUDE
+
+#if PLATFORM_MICROSOFT
+#include "Microsoft/HideMicrosoftPlatformTypes.h"
+#endif
+
+#endif //defined(PLATFORM_CURL_INCLUDE)
+
+
 class FCompositeBuffer;
 class FCbObjectView;
 class FCbPackage;
@@ -180,7 +196,11 @@ public:
 	ZEN_API bool GetHeader(const ANSICHAR* Header, FString& OutValue) const;
 
 private:
+#if defined(CURL_STRICTER)
+	CURL*					Curl = nullptr;
+#else
 	void* /* CURL */		Curl = nullptr;
+#endif
 	long /* CURLCode */		CurlResult;
 	long					ResponseCode = 0;
 	size_t					BytesSent = 0;
