Ensure that zenserver is launched in ASIO mode even for licensee configurations.  Avoid forced persistence of the zen data path.

diff --git a/Engine/Config/BaseEngine.ini b/Engine/Config/BaseEngine.ini
--- a/Engine/Config/BaseEngine.ini
+++ b/Engine/Config/BaseEngine.ini
@@ -2080,6 +2080,7 @@ AutoLaunch=true
 DataPath=%ENGINEVERSIONAGNOSTICUSERDIR%Zen
 LocalDataCachePathEnvOverride=UE-LocalDataCachePath
 LocalDataCachePathEditorOverrideSetting=LocalDerivedDataCache
+ExtraArgs=--http asio
 
 [Zen.ConnectExisting]
 HostName=localhost
diff --git a/Engine/Source/Developer/Zen/Internal/ZenServerInterface.cpp b/Engine/Source/Developer/Zen/Internal/ZenServerInterface.cpp
--- a/Engine/Source/Developer/Zen/Internal/ZenServerInterface.cpp
+++ b/Engine/Source/Developer/Zen/Internal/ZenServerInterface.cpp
@@ -134,13 +134,6 @@ DetermineDataPath(const TCHAR* ConfigSection, FString& DataPath)
 		}
 	}
 
-	ON_SCOPE_EXIT
-	{
-		// Persist the data path so that future runs use the same one.
-		DataPath = NormalizeDataPath(DataPath);
-		FPlatformMisc::SetStoredValue(TEXT("Epic Games"), TEXT("Zen"), TEXT("DataPath"), DataPath);
-	};
-
 	// Zen environment
 	DataPathEnvOverrideValue = FPlatformMisc::GetEnvironmentVariable(TEXT("UE-ZenDataPath"));
 	if (!DataPathEnvOverrideValue.IsEmpty())
