Autolaunch fixes ahead of wider ZenServer usage/rollout:
-Enable ZenServer GC on 6 hour interval when autolaunching ZenServer
-Change default ZenServer data path to be C:\ProgramData\Epic\Zen
-Ensure when commandline changes that ZenServer is re-launched by the autolauncher
-Handle possibility of existing ZenServer being on the same port but at different data directory than the desired autolaunch target
-Handle possibility of existing ZenServer being at same data data directory but different port than the desired autolaunch target

diff --git a/Engine/Config/BaseEngine.ini b/Engine/Config/BaseEngine.ini
--- a/Engine/Config/BaseEngine.ini
+++ b/Engine/Config/BaseEngine.ini
@@ -2077,10 +2077,10 @@ PathToRAC=RuntimeAssetCache
 AutoLaunch=true
 
 [Zen.AutoLaunch]
-DataPath=%ENGINEVERSIONAGNOSTICUSERDIR%Zen
+DataPath=%APPSETTINGSDIR%Zen
 LocalDataCachePathEnvOverride=UE-LocalDataCachePath
 LocalDataCachePathEditorOverrideSetting=LocalDerivedDataCache
-ExtraArgs=--http asio
+ExtraArgs=--http asio --gc-interval-seconds 21600
 
 [Zen.ConnectExisting]
 HostName=localhost
diff --git a/Engine/Source/Developer/Zen/Internal/ZenServerInterface.cpp b/Engine/Source/Developer/Zen/Internal/ZenServerInterface.cpp
--- a/Engine/Source/Developer/Zen/Internal/ZenServerInterface.cpp
+++ b/Engine/Source/Developer/Zen/Internal/ZenServerInterface.cpp
@@ -10,10 +10,12 @@
 #include "HAL/FileManager.h"
 #include "HAL/Platform.h"
 #include "HAL/PlatformMisc.h"
+#include "HAL/PlatformProcess.h"
 #include "HAL/PlatformTime.h"
 #include "Logging/LogScopedCategoryAndVerbosityOverride.h"
 #include "Misc/App.h"
 #include "Misc/ConfigCacheIni.h"
+#include "Misc/FileHelper.h"
 #include "Misc/Paths.h"
 #include "Misc/PathViews.h"
 #include "Misc/ScopedSlowTask.h"
@@ -387,6 +389,75 @@ ReadCbLockFile(FStringView FileName, FCbObject& OutLockObject)
 #endif
 }
 
+static void
+RequestZenShutdownOnPort(uint16 Port)
+{
+#if PLATFORM_WINDOWS
+	HANDLE Handle = OpenEventW(EVENT_MODIFY_STATE, false, *WriteToWideString<64>(WIDETEXT("Zen_"), Port, WIDETEXT("_Shutdown")));
+	if (Handle != INVALID_HANDLE_VALUE)
+	{
+		ON_SCOPE_EXIT{ CloseHandle(Handle); };
+		SetEvent(Handle);
+	}
+#else
+	static_assert(false, "Missing implementation for Zen named shutdown events");
+#endif
+}
+
+static bool
+WaitForZenShutdown(const TCHAR* LockFilePath, double MaximumWaitDurationSeconds)
+{
+	uint64 ZenShutdownWaitStartTime = FPlatformTime::Cycles64();
+	while (IFileManager::Get().FileExists(LockFilePath))
+	{
+		double ZenShutdownWaitDuration = FPlatformTime::ToSeconds64(FPlatformTime::Cycles64() - ZenShutdownWaitStartTime);
+		if (ZenShutdownWaitDuration < MaximumWaitDurationSeconds)
+		{
+			FPlatformProcess::Sleep(0.01f);
+		}
+		else
+		{
+			return false;
+		}
+	}
+	return true;
+}
+
+static bool IsProcessActive(const TCHAR* ExecutablePath)
+{
+	FString NormalizedExecutablePath(ExecutablePath);
+	FPaths::NormalizeFilename(NormalizedExecutablePath);
+	FPlatformProcess::FProcEnumerator ProcIter;
+	while (ProcIter.MoveNext())
+	{
+		FPlatformProcess::FProcEnumInfo ProcInfo = ProcIter.GetCurrent();
+		FString Candidate = ProcInfo.GetFullPath();
+		FPaths::NormalizeFilename(Candidate);
+		if (Candidate == NormalizedExecutablePath)
+		{
+			return true;
+		}
+	}
+	return false;
+}
+
+static FString
+DetermineCmdLineWithoutTransientComponents(const FServiceAutoLaunchSettings& InSettings, int16 OverrideDesiredPort)
+{
+	FString Parms;
+	Parms.Appendf(TEXT("--port %d --data-dir \"%s\""),
+		OverrideDesiredPort,
+		*InSettings.DataPath);
+
+	if (!InSettings.ExtraArgs.IsEmpty())
+	{
+		Parms.AppendChar(TEXT(' '));
+		Parms.Append(InSettings.ExtraArgs);
+	}
+
+	return Parms;
+}
+
 static bool GIsDefaultServicePresent = false;
 
 FZenServiceInstance& GetDefaultServiceInstance()
@@ -532,7 +603,7 @@ FZenServiceInstance::ConditionalUpdateLocalInstall()
 	FileManager.GetTimeStampPair(*InTreeFilePath, *InstallFilePath, InTreeFileTime, InstallFileTime);
 	if (InTreeFileTime > InstallFileTime)
 	{
-		if (FPlatformProcess::IsApplicationRunning(*FPaths::GetCleanFilename(InstallFilePath)))
+		if (IsProcessActive(*InstallFilePath))
 		{
 			// TODO: Instead of using the lock file, this could use the shared memory system state named "Global\ZenMap" (see zenserverprocess.{h,cpp} in zen codebase)
 			FString LockFilePath = FPaths::Combine(Settings.SettingsVariant.Get<FServiceAutoLaunchSettings>().DataPath, TEXT(".lock"));
@@ -542,31 +613,9 @@ FZenServiceInstance::ConditionalUpdateLocalInstall()
 				uint16 RunningPort = LockObject["port"].AsUInt16(0);
 				if (RunningPort != 0)
 				{
-#if PLATFORM_WINDOWS
-					FString ShutdownEventName = *WriteToString<64>(TEXT("Zen_"), RunningPort, TEXT("_Shutdown"));
-					HANDLE Handle = OpenEventW(EVENT_MODIFY_STATE, false, *ShutdownEventName);
-					if (Handle != INVALID_HANDLE_VALUE)
-					{
-						ON_SCOPE_EXIT{ CloseHandle(Handle); };
-						SetEvent(Handle);
-					}
-#else
-					static_assert(false, "Missing implementation for Zen named shutdown events");
-#endif
+					RequestZenShutdownOnPort(RunningPort);
 
-					uint64 ZenShutdownWaitStartTime = FPlatformTime::Cycles64();
-					while (FileManager.FileExists(*LockFilePath))
-					{
-						double ZenShutdownWaitDuration = FPlatformTime::ToSeconds64(FPlatformTime::Cycles64() - ZenShutdownWaitStartTime);
-						if (ZenShutdownWaitDuration < 5.0)
-						{
-							FPlatformProcess::Sleep(0.01f);
-						}
-						else
-						{
-							break;
-						}
-					}
+					WaitForZenShutdown(*LockFilePath, 5.0);
 				}
 			}
 
@@ -595,26 +644,57 @@ FZenServiceInstance::ConditionalUpdateLocalInstall()
 bool
 FZenServiceInstance::AutoLaunch(const FServiceAutoLaunchSettings& InSettings, FString&& ExecutablePath, FString& OutHostName, uint16& OutPort)
 {
+	int16 DesiredPort = InSettings.DesiredPort;
 	IFileManager& FileManager = IFileManager::Get();
-	FString LockFilePath = FPaths::Combine(InSettings.DataPath, TEXT(".lock"));
+	const FString LockFilePath = FPaths::Combine(InSettings.DataPath, TEXT(".lock"));
+	const FString CmdLineFilePath = FPaths::Combine(InSettings.DataPath, TEXT(".cmdline"));
 	FileManager.Delete(*LockFilePath, false, false, true);
 
-	bool bProcessIsLive = FileManager.FileExists(*LockFilePath) && FPlatformProcess::IsApplicationRunning(*FPaths::GetCleanFilename(ExecutablePath));
+	bool bReUsingExistingInstance = false;
 
-	// When limiting process lifetime, always re-launch to add sponsor process IDs.
-	// When not limiting process lifetime, only launch if the process is not already live.
-	if (InSettings.bLimitProcessLifetime || !bProcessIsLive)
+	if (FileManager.FileExists(*LockFilePath))
 	{
-		FString Parms;
-		Parms.Appendf(TEXT("--port %d --data-dir \"%s\""),
-			InSettings.DesiredPort,
-			*InSettings.DataPath);
+		// If an instance is running with this data path, check if we can use it and what port it is on
+		uint16 CurrentPort = 0;
+		FCbObject LockObject;
+		if (ReadCbLockFile(LockFilePath, LockObject))
+		{
+			bool bIsReady = LockObject["ready"].AsBool();
+			if (bIsReady)
+			{
+				CurrentPort = LockObject["port"].AsUInt16();
+			}
+		}
 
-		if (InSettings.bLimitProcessLifetime || GIsBuildMachine)
+		bool bCurrentInstanceUsable = false;
+		FString DesiredCmdLine = DetermineCmdLineWithoutTransientComponents(InSettings, CurrentPort);
+		FString CurrentCmdLine;
+		if (FFileHelper::LoadFileToString(CurrentCmdLine, *CmdLineFilePath) && (DesiredCmdLine == CurrentCmdLine))
 		{
-			Parms.Appendf(TEXT(" --owner-pid %d"),
-				FPlatformProcess::GetCurrentProcessId());
+			DesiredPort = CurrentPort;
+			bReUsingExistingInstance = true;
 		}
+		else
+		{
+			RequestZenShutdownOnPort(CurrentPort);
+			WaitForZenShutdown(*LockFilePath, 5.0);
+		}
+	}
+
+	if (!bReUsingExistingInstance)
+	{
+		RequestZenShutdownOnPort(DesiredPort);
+	}
+
+
+	bool bProcessIsLive = FileManager.FileExists(*LockFilePath);
+
+	// When limiting process lifetime, always re-launch to add sponsor process IDs.
+	// When not limiting process lifetime, only launch if the process is not already live.
+	if (InSettings.bLimitProcessLifetime || !bProcessIsLive)
+	{
+		FString ParmsWithoutTransients= DetermineCmdLineWithoutTransientComponents(InSettings, DesiredPort);
+		FString Parms = ParmsWithoutTransients;
 
 		FString LogCommandLineOverrideValue;
 		if (FParse::Value(FCommandLine::Get(), TEXT("ZenLogPath="), LogCommandLineOverrideValue))
@@ -626,10 +706,10 @@ FZenServiceInstance::AutoLaunch(const FServiceAutoLaunchSettings& InSettings, FS
 			}
 		}
 
-		if (!InSettings.ExtraArgs.IsEmpty())
+		if (InSettings.bLimitProcessLifetime || GIsBuildMachine)
 		{
-			Parms.AppendChar(TEXT(' '));
-			Parms.Append(InSettings.ExtraArgs);
+			Parms.Appendf(TEXT(" --owner-pid %d"),
+				FPlatformProcess::GetCurrentProcessId());
 		}
 
 		FProcHandle Proc;
@@ -701,13 +781,18 @@ FZenServiceInstance::AutoLaunch(const FServiceAutoLaunchSettings& InSettings, FS
 				PipeReadChild);
 		}
 #endif
+		if (!bProcessIsLive)
+		{
+			FFileHelper::SaveStringToFile(ParmsWithoutTransients,*CmdLineFilePath);
+		}
+
 		bProcessIsLive = Proc.IsValid();
 	}
 
 
 	OutHostName = TEXT("localhost");
 	// Default to assuming that we get to run on the port we want
-	OutPort = InSettings.DesiredPort;
+	OutPort = DesiredPort;
 
 	if (bProcessIsLive)
 	{
@@ -729,7 +814,7 @@ FZenServiceInstance::AutoLaunch(const FServiceAutoLaunchSettings& InSettings, FS
 				bIsReady = LockObject["ready"].AsBool();
 				if (bIsReady)
 				{
-					OutPort = LockObject["port"].AsUInt16(InSettings.DesiredPort);
+					OutPort = LockObject["port"].AsUInt16(DesiredPort);
 					break;
 				}
 			}
