Remove coded special handling for build machines and lifetime limitation of ZenServer.  Instead have build scripts that want that behavior explicitly specify it as an INI override on the commandline.  Misc fixes to commandline and path consistency when launching ZenServer.

diff --git a/Engine/Source/Developer/Zen/Internal/ZenServerInterface.cpp b/Engine/Source/Developer/Zen/Internal/ZenServerInterface.cpp
--- a/Engine/Source/Developer/Zen/Internal/ZenServerInterface.cpp
+++ b/Engine/Source/Developer/Zen/Internal/ZenServerInterface.cpp
@@ -140,7 +140,7 @@ DetermineDataPath(const TCHAR* ConfigSection, FString& DataPath)
 	DataPathEnvOverrideValue = FPlatformMisc::GetEnvironmentVariable(TEXT("UE-ZenDataPath"));
 	if (!DataPathEnvOverrideValue.IsEmpty())
 	{
-		DataPath = DataPathEnvOverrideValue;
+		DataPath = NormalizeDataPath(DataPathEnvOverrideValue);
 		UE_LOG(LogZenServiceInstance, Log, TEXT("Found environment variable UE-ZenDataPath=%s"), *DataPathEnvOverrideValue);
 		return;
 	}
@@ -150,12 +150,13 @@ DetermineDataPath(const TCHAR* ConfigSection, FString& DataPath)
 	DetermineLocalDataCachePath(ConfigSection, LocalDataCachePath);
 	if (!LocalDataCachePath.IsEmpty() && (LocalDataCachePath != TEXT("None")) && !FPaths::IsUnderDirectory(LocalDataCachePath, FPaths::RootDir()))
 	{
-		DataPath = FPaths::Combine(LocalDataCachePath, TEXT("Zen"));
+		DataPath = NormalizeDataPath(FPaths::Combine(LocalDataCachePath, TEXT("Zen")));
 		return;
 	}
 
 	// Zen config default
 	GConfig->GetString(ConfigSection, TEXT("DataPath"), DataPath, GEngineIni);
+	DataPath = NormalizeDataPath(DataPath);
 
 	check(!DataPath.IsEmpty())
 }
@@ -224,15 +225,8 @@ FServiceSettings::ReadFromJson(FJsonObject& JsonObject)
 					AutoLaunchSettings.DataPath = AutoLaunchSettingsObject->Values.FindRef(TEXT("DataPath"))->AsString();
 					AutoLaunchSettings.ExtraArgs = AutoLaunchSettingsObject->Values.FindRef(TEXT("ExtraArgs"))->AsString();
 					AutoLaunchSettingsObject->Values.FindRef(TEXT("DesiredPort"))->TryGetNumber(AutoLaunchSettings.DesiredPort);
-					AutoLaunchSettingsObject->Values.FindRef(TEXT("bShowConsole"))->TryGetBool(AutoLaunchSettings.bShowConsole);
-					AutoLaunchSettingsObject->Values.FindRef(TEXT("bLimitProcessLifetime"))->TryGetBool(AutoLaunchSettings.bLimitProcessLifetime);
-					if (TSharedPtr<FJsonValue> TreatAsBuildMachineValue = AutoLaunchSettingsObject->Values.FindRef(TEXT("TreatAsBuildMachine")))
-					{
-						if (TreatAsBuildMachineValue->AsString() == FPlatformProcess::ComputerName())
-						{
-							AutoLaunchSettings.bLimitProcessLifetime = true;
-						}
-					}
+					AutoLaunchSettingsObject->Values.FindRef(TEXT("ShowConsole"))->TryGetBool(AutoLaunchSettings.bShowConsole);
+					AutoLaunchSettingsObject->Values.FindRef(TEXT("LimitProcessLifetime"))->TryGetBool(AutoLaunchSettings.bLimitProcessLifetime);
 				}
 			}
 		}
@@ -290,12 +284,8 @@ FServiceSettings::WriteToJson(TJsonWriter<TCHAR, TPrettyJsonPrintPolicy<TCHAR>>&
 		Writer.WriteValue(TEXT("DataPath"), AutoLaunchSettings.DataPath);
 		Writer.WriteValue(TEXT("ExtraArgs"), AutoLaunchSettings.ExtraArgs);
 		Writer.WriteValue(TEXT("DesiredPort"), AutoLaunchSettings.DesiredPort);
-		Writer.WriteValue(TEXT("bShowConsole"), AutoLaunchSettings.bShowConsole);
-		Writer.WriteValue(TEXT("bLimitProcessLifetime"), AutoLaunchSettings.bLimitProcessLifetime);
-		if (GIsBuildMachine)
-		{
-			Writer.WriteValue(TEXT("TreatAsBuildMachine"), FPlatformProcess::ComputerName());
-		}
+		Writer.WriteValue(TEXT("ShowConsole"), AutoLaunchSettings.bShowConsole);
+		Writer.WriteValue(TEXT("LimitProcessLifetime"), AutoLaunchSettings.bLimitProcessLifetime);
 		Writer.WriteObjectEnd();
 	}
 	else
@@ -444,10 +434,13 @@ static bool IsProcessActive(const TCHAR* ExecutablePath)
 static FString
 DetermineCmdLineWithoutTransientComponents(const FServiceAutoLaunchSettings& InSettings, int16 OverrideDesiredPort)
 {
+	FString PlatformDataPath(InSettings.DataPath);
+	FPaths::MakePlatformFilename(PlatformDataPath);
+
 	FString Parms;
 	Parms.Appendf(TEXT("--port %d --data-dir \"%s\""),
 		OverrideDesiredPort,
-		*InSettings.DataPath);
+		*PlatformDataPath);
 
 	if (!InSettings.ExtraArgs.IsEmpty())
 	{
@@ -595,7 +588,7 @@ FString
 FZenServiceInstance::ConditionalUpdateLocalInstall()
 {
 	FString InTreeFilePath = FPaths::ConvertRelativePathToFull(FPlatformProcess::GenerateApplicationPath(TEXT("zenserver"), EBuildConfiguration::Development));
-	FString InstallFilePath = FPaths::ConvertRelativePathToFull(FPaths::Combine(FPaths::EngineSavedDir(), TEXT("ZenServer"), FString(FPathViews::GetCleanFilename(InTreeFilePath))));
+	FString InstallFilePath = FPaths::ConvertRelativePathToFull(FPaths::Combine(FPlatformProcess::ApplicationSettingsDir(), TEXT("Zen\\Install"), FString(FPathViews::GetCleanFilename(InTreeFilePath))));
 
 	IFileManager& FileManager = IFileManager::Get();
 	FDateTime InTreeFileTime;
@@ -638,6 +631,7 @@ FZenServiceInstance::ConditionalUpdateLocalInstall()
 #endif
 	}
 
+	FPaths::MakePlatformFilename(InstallFilePath);
 	return InstallFilePath;
 }
 
@@ -667,7 +661,7 @@ FZenServiceInstance::AutoLaunch(const FServiceAutoLaunchSettings& InSettings, FS
 		}
 
 		bool bCurrentInstanceUsable = false;
-		FString DesiredCmdLine = DetermineCmdLineWithoutTransientComponents(InSettings, CurrentPort);
+		FString DesiredCmdLine = FString::Printf(TEXT("%s %s"), *ExecutablePath, *DetermineCmdLineWithoutTransientComponents(InSettings, CurrentPort));
 		FString CurrentCmdLine;
 		if (FFileHelper::LoadFileToString(CurrentCmdLine, *CmdLineFilePath) && (DesiredCmdLine == CurrentCmdLine))
 		{
@@ -693,7 +687,7 @@ FZenServiceInstance::AutoLaunch(const FServiceAutoLaunchSettings& InSettings, FS
 	// When not limiting process lifetime, only launch if the process is not already live.
 	if (InSettings.bLimitProcessLifetime || !bProcessIsLive)
 	{
-		FString ParmsWithoutTransients= DetermineCmdLineWithoutTransientComponents(InSettings, DesiredPort);
+		FString ParmsWithoutTransients = DetermineCmdLineWithoutTransientComponents(InSettings, DesiredPort);
 		FString Parms = ParmsWithoutTransients;
 
 		FString LogCommandLineOverrideValue;
@@ -706,7 +700,7 @@ FZenServiceInstance::AutoLaunch(const FServiceAutoLaunchSettings& InSettings, FS
 			}
 		}
 
-		if (InSettings.bLimitProcessLifetime || GIsBuildMachine)
+		if (InSettings.bLimitProcessLifetime)
 		{
 			Parms.Appendf(TEXT(" --owner-pid %d"),
 				FPlatformProcess::GetCurrentProcessId());
@@ -783,7 +777,8 @@ FZenServiceInstance::AutoLaunch(const FServiceAutoLaunchSettings& InSettings, FS
 #endif
 		if (!bProcessIsLive)
 		{
-			FFileHelper::SaveStringToFile(ParmsWithoutTransients,*CmdLineFilePath);
+			FString ExecutedCmdLine = FString::Printf(TEXT("%s %s"), *ExecutablePath, *ParmsWithoutTransients);
+			FFileHelper::SaveStringToFile(ExecutedCmdLine,*CmdLineFilePath);
 		}
 
 		bProcessIsLive = Proc.IsValid();
