Support cooking to remote Zen server

#rnx

diff --git a/Engine/Source/Developer/IoStoreUtilities/Private/ZenStoreHttpClient.cpp b/Engine/Source/Developer/IoStoreUtilities/Private/ZenStoreHttpClient.cpp
--- a/Engine/Source/Developer/IoStoreUtilities/Private/ZenStoreHttpClient.cpp
+++ b/Engine/Source/Developer/IoStoreUtilities/Private/ZenStoreHttpClient.cpp
@@ -321,6 +321,8 @@ TIoStatusOr<uint64> FZenStoreHttpClient::AppendOp(FCbPackage OpEntry)
 
 			// Send phase
 
+			const bool bCanUseLocalTempFile = ZenService.GetInstance().IsServiceRunningLocally();
+
 			for (const FCbAttachment& Attachment : Attachments)
 			{
 				if (!Attachment.IsCompressedBinary())
@@ -334,7 +336,7 @@ TIoStatusOr<uint64> FZenStoreHttpClient::AppendOp(FCbPackage OpEntry)
 				if (NeedChunks.Contains(AttachmentHash))
 				{
 					FSharedBuffer AttachmentData = Attachment.AsCompressedBinary().GetCompressed().ToShared();
-					if (AttachmentData.GetSize() >= StandaloneThresholdBytes)
+					if (bCanUseLocalTempFile && AttachmentData.GetSize() >= StandaloneThresholdBytes)
 					{
 						// Write to temporary file. To avoid race conditions we derive
 						// the file name from a salt value and the attachment hash
diff --git a/Engine/Source/Developer/Zen/Internal/ZenServerInterface.cpp b/Engine/Source/Developer/Zen/Internal/ZenServerInterface.cpp
--- a/Engine/Source/Developer/Zen/Internal/ZenServerInterface.cpp
+++ b/Engine/Source/Developer/Zen/Internal/ZenServerInterface.cpp
@@ -24,6 +24,7 @@
 #include "Serialization/CompactBinarySerialization.h"
 #include "Serialization/CompactBinaryValidation.h"
 #include "String/LexFromString.h"
+#include "SocketSubsystem.h"
 
 #if PLATFORM_WINDOWS
 #	include "Windows/AllowWindowsPlatformTypes.h"
@@ -169,6 +170,46 @@ ReadUInt16FromConfig(const TCHAR* Section, const TCHAR* Key, uint16& Value, cons
 	Value = (uint16)ValueInt32;
 }
 
+static bool
+IsLocalHost(const FString& Host)
+{
+	if (Host.Compare(FString(TEXT("localhost")), ESearchCase::IgnoreCase) == 0)
+	{
+		return true;
+	}
+
+	if (Host.Compare(FString(TEXT("127.0.0.1"))) == 0)
+	{
+		return true;
+	}
+
+	ISocketSubsystem& SocketSubsystem = *ISocketSubsystem::Get();
+
+	const TSharedPtr<FInternetAddr> Addr = SocketSubsystem.GetAddressFromString(Host);
+	if (!Addr)
+	{
+		UE_LOG(LogZenServiceInstance, Warning, TEXT("Failed to get internet address from host '%s'"), *Host);
+		return false;
+	}
+
+	TArray<TSharedPtr<FInternetAddr>> LocalAddresses;
+	if (!SocketSubsystem.GetLocalAdapterAddresses(LocalAddresses))
+	{
+		UE_LOG(LogZenServiceInstance, Warning, TEXT("Failed to find local adapter addresses"));
+		return false;
+	}
+
+	for (const auto& Local : LocalAddresses)
+	{
+		if (*Local == *Addr)
+		{
+			return true;
+		}
+	}
+
+	return false;
+}
+
 void
 FServiceSettings::ReadFromConfig()
 {
@@ -310,6 +351,27 @@ FServiceSettings::TryApplyAutoLaunchOverride()
 		ConnectExistingSettings.Port = 1337;
 		return true;
 	}
+
+	FString Host;
+	if  (FParse::Value(FCommandLine::Get(), TEXT("-NoZenAutoLaunch="), Host))
+	{
+		SettingsVariant.Emplace<FServiceConnectSettings>();
+		FServiceConnectSettings& ConnectExistingSettings = SettingsVariant.Get<FServiceConnectSettings>();
+
+		int32 PortDelimIndex = INDEX_NONE;
+		if (Host.FindChar(TEXT(':'), PortDelimIndex))
+		{
+			ConnectExistingSettings.HostName = Host.Left(PortDelimIndex);
+			LexFromString(ConnectExistingSettings.Port, Host.RightChop(PortDelimIndex + 1));
+		}
+		else
+		{
+			ConnectExistingSettings.HostName = Host;
+			ConnectExistingSettings.Port = 1337;
+		}
+
+		return true;
+	}
 #endif
 	return false;
 }
@@ -559,6 +621,7 @@ FZenServiceInstance::Initialize()
 		if (bHasLaunchedLocal)
 		{
 			AutoLaunchedPort = Port;
+			bIsRunningLocally = true;
 		}
 	}
 	else
@@ -566,6 +629,7 @@ FZenServiceInstance::Initialize()
 		const FServiceConnectSettings& ConnectExistingSettings = Settings.SettingsVariant.Get<FServiceConnectSettings>();
 		HostName = ConnectExistingSettings.HostName;
 		Port = ConnectExistingSettings.Port;
+		bIsRunningLocally = IsLocalHost(HostName);
 	}
 	URL = WriteToString<64>(TEXT("http://"), HostName, TEXT(":"), Port, TEXT("/"));
 }
diff --git a/Engine/Source/Developer/Zen/Public/ZenServerInterface.h b/Engine/Source/Developer/Zen/Public/ZenServerInterface.h
--- a/Engine/Source/Developer/Zen/Public/ZenServerInterface.h
+++ b/Engine/Source/Developer/Zen/Public/ZenServerInterface.h
@@ -126,6 +126,7 @@ public:
 	UE_API bool GetStats(FZenStats& stats) const;
 	UE_API bool IsServiceRunning();
 	UE_API bool IsServiceReady();
+	UE_API bool IsServiceRunningLocally() const { return bIsRunningLocally; }
 
 	static uint16 GetAutoLaunchedPort() { return AutoLaunchedPort; }
 
@@ -142,6 +143,7 @@ private:
 	uint16 Port;
 	static inline uint16 AutoLaunchedPort = 0;
 	bool bHasLaunchedLocal = false;
+	bool bIsRunningLocally = true;
 };
 
 } // namespace UE::Zen
diff --git a/Engine/Source/Developer/Zen/Zen.Build.cs b/Engine/Source/Developer/Zen/Zen.Build.cs
--- a/Engine/Source/Developer/Zen/Zen.Build.cs
+++ b/Engine/Source/Developer/Zen/Zen.Build.cs
@@ -7,10 +7,7 @@ public class Zen : ModuleRules
 {
 	public Zen(ReadOnlyTargetRules Target) : base(Target)
 	{
-		PrivateDependencyModuleNames.Add("Core");
-
-		// Dependencies for "S3" and "HTTP" backends
-		PrivateDependencyModuleNames.AddRange(new string[] { "SSL", "Json" });
+		PrivateDependencyModuleNames.AddRange(new string[] { "Core", "Sockets", "SSL", "Json" });
 		PrivateIncludePathModuleNames.Add("DesktopPlatform");
 		AddEngineThirdPartyPrivateStaticDependencies(Target, "libcurl");
 		AddEngineThirdPartyPrivateStaticDependencies(Target, "OpenSSL");
