Zen - updated stats response to match new Zen server binaries

diff --git a/Engine/Source/Developer/Zen/Internal/ZenServerInterface.cpp b/Engine/Source/Developer/Zen/Internal/ZenServerInterface.cpp
--- a/Engine/Source/Developer/Zen/Internal/ZenServerInterface.cpp
+++ b/Engine/Source/Developer/Zen/Internal/ZenServerInterface.cpp
@@ -991,12 +991,16 @@ FZenServiceInstance::GetStats( FZenStats& stats ) const
 
 			EndPointStats.Name = FString(EndPointView["name"].AsString());
 			EndPointStats.Url = FString(EndPointView["url"].AsString());
-			EndPointStats.Health = FString(EndPointView["health"].AsString());
-			EndPointStats.HitRatio = EndPointView["hit_ratio"].AsDouble();
-			EndPointStats.UploadedMB = EndPointView["uploaded_mb"].AsDouble();
-			EndPointStats.DownloadedMB = EndPointView["downloaded_mb"].AsDouble();
-			EndPointStats.ErrorCount = EndPointView["error_count"].AsInt64();
+			EndPointStats.Health = FString(EndPointView["state"].AsString());
 
+			if (FCbObjectView Cache = EndPointView["cache"].AsObjectView())
+			{
+				EndPointStats.HitRatio = Cache["hit_ratio"].AsDouble();
+				EndPointStats.UploadedMB = Cache["put_bytes"].AsDouble() / 1024.0 / 1024.0;
+				EndPointStats.DownloadedMB = Cache["get_bytes"].AsDouble() / 1024.0 / 1024.0;
+				EndPointStats.ErrorCount = Cache["error_count"].AsInt64();
+			}
+			
 			UpstreamStats.TotalUploadedMB += EndPointStats.UploadedMB;
 			UpstreamStats.TotalDownloadedMB += EndPointStats.DownloadedMB;
 
