Fix bug with ZenServer AutoLaunch port being inlined into each module.  Also fixed failure to distinguish between "did not launch" and "taking long to be ready" outcomes in Zen AutoLaunch.

diff --git a/Engine/Source/Developer/Zen/Internal/ZenServerInterface.cpp b/Engine/Source/Developer/Zen/Internal/ZenServerInterface.cpp
--- a/Engine/Source/Developer/Zen/Internal/ZenServerInterface.cpp
+++ b/Engine/Source/Developer/Zen/Internal/ZenServerInterface.cpp
@@ -378,6 +378,8 @@ FServiceSettings::TryApplyAutoLaunchOverride()
 
 #if UE_WITH_ZEN
 
+uint16 FZenServiceInstance::AutoLaunchedPort = 0;
+
 static bool
 ReadCbLockFile(FStringView FileName, FCbObject& OutLockObject)
 {
@@ -879,18 +881,36 @@ FZenServiceInstance::AutoLaunch(const FServiceAutoLaunchSettings& InSettings, FS
 			}
 
 			double ZenWaitDuration = FPlatformTime::ToSeconds64(FPlatformTime::Cycles64() - ZenWaitStartTime);
-			if (ZenWaitDuration < 1.0)
+			if (ZenWaitDuration < 3.0)
 			{
-				// Initial 1 second window of higher frequency checks
+				// Initial 3 second window of higher frequency checks
 				FPlatformProcess::Sleep(0.01f);
 			}
 			else
 			{
 				if (DurationPhase == EWaitDurationPhase::Short)
 				{
-					// Insist after 1 second that the lock file must exist by now otherwise there has been a failure to launch.
-					// The file existing does not mean that it is ready to use, just that it is running.  It may be performing scrubbing.
-					checkf(FileManager.FileExists(*LockFilePath), TEXT("ZenServer did not launch in the expected duration."));// Note that the dialog may not show up when zenserver is needed early in the launch cycle, but this will at least ensure
+					if (!FileManager.FileExists(*LockFilePath))
+					{
+						if (FApp::IsUnattended())
+						{
+							checkf(false, TEXT("ZenServer did not launch in the expected duration."));
+						}
+						else
+						{
+							FText ZenLaunchFailurePromptTitle = NSLOCTEXT("Zen", "Zen_LaunchFailurePromptTitle", "Failed to launch");
+
+							FFormatNamedArguments FormatArguments;
+							FString LogFilePath = FPaths::Combine(InSettings.DataPath, TEXT("logs"), TEXT("zenserver.log"));
+							FPaths::MakePlatformFilename(LogFilePath);
+							FormatArguments.Add(TEXT("LogFilePath"), FText::FromString(LogFilePath));
+							FText ZenLaunchFailurePromptText = FText::Format(NSLOCTEXT("Zen", "Zen_LaunchFailurePromptText", "ZenServer failed to launch. This process will now exit. Please check the ZenServer log file for details:\n{LogFilePath}"), FormatArguments);
+							FPlatformMisc::MessageBoxExt(EAppMsgType::Ok, *ZenLaunchFailurePromptText.ToString(), *ZenLaunchFailurePromptTitle.ToString());
+							FPlatformMisc::RequestExit(true);
+							return false;
+						}
+					}
+					// Note that the dialog may not show up when zenserver is needed early in the launch cycle, but this will at least ensure
 					// the splash screen is refreshed with the appropriate text status message.
 					WaitForZenReadySlowTask.MakeDialog(true, false);
 					UE_LOG(LogZenServiceInstance, Display, TEXT("Waiting for ZenServer to be ready..."));
diff --git a/Engine/Source/Developer/Zen/Public/ZenServerInterface.h b/Engine/Source/Developer/Zen/Public/ZenServerInterface.h
--- a/Engine/Source/Developer/Zen/Public/ZenServerInterface.h
+++ b/Engine/Source/Developer/Zen/Public/ZenServerInterface.h
@@ -128,7 +128,7 @@ public:
 	UE_API bool IsServiceReady();
 	UE_API bool IsServiceRunningLocally() const { return bIsRunningLocally; }
 
-	static uint16 GetAutoLaunchedPort() { return AutoLaunchedPort; }
+	static UE_API uint16 GetAutoLaunchedPort() { return AutoLaunchedPort; }
 
 private:
 
@@ -141,7 +141,7 @@ private:
 	FString URL;
 	FString HostName;
 	uint16 Port;
-	static inline uint16 AutoLaunchedPort = 0;
+	static uint16 AutoLaunchedPort;
 	bool bHasLaunchedLocal = false;
 	bool bIsRunningLocally = true;
 };
