ZenServerHttp: Add support for rpc calls that post an FCbPackage; previously only FCbObject was supported.

diff --git a/Engine/Source/Developer/Zen/Internal/ZenServerHttp.cpp b/Engine/Source/Developer/Zen/Internal/ZenServerHttp.cpp
--- a/Engine/Source/Developer/Zen/Internal/ZenServerHttp.cpp
+++ b/Engine/Source/Developer/Zen/Internal/ZenServerHttp.cpp
@@ -130,12 +130,13 @@ namespace UE::Zen {
 		return PerformBlocking(Uri, RequestVerb::Put, ContentLength);
 	}
 
-	FZenHttpRequest::Result FZenHttpRequest::PerformBlockingPost(FStringView Uri, FCbObjectView Obj)
+	FZenHttpRequest::Result FZenHttpRequest::PerformBlockingPost(FStringView Uri, FCbObjectView Obj,
+		EContentType AcceptType)
 	{
 		FLargeMemoryWriter Out;
 		Obj.CopyTo(Out);
 
-		return PerformBlockingPost(Uri, Out.GetView(), EContentType::CbObject);
+		return PerformBlockingPost(Uri, Out.GetView(), EContentType::CbObject, AcceptType);
 	}
 
 	struct CbPackageHeader
@@ -162,10 +163,10 @@ namespace UE::Zen {
 
 	};
 
-	FZenHttpRequest::Result FZenHttpRequest::PerformBlockingPostPackage(FStringView Uri, FCbPackage Package)
+	FZenHttpRequest::Result FZenHttpRequest::PerformBlockingPostPackage(FStringView Uri, const FCbPackage& Package, EContentType AcceptType)
 	{
 		TConstArrayView<FCbAttachment> Attachments = Package.GetAttachments();
-		FCbObject Object = Package.GetObject();
+		const FCbObject& Object = Package.GetObject();
 		FCompressedBuffer ObjectBuffer = FCompressedBuffer::Compress(Object.GetBuffer(), FOodleDataCompression::ECompressor::NotSet, FOodleDataCompression::ECompressionLevel::None);
 
 		CbPackageHeader Hdr;
@@ -256,30 +257,24 @@ namespace UE::Zen {
 			}
 		}
 
-		return PerformBlockingPost(Uri, Out.GetView(), EContentType::CbPackage);
+		return PerformBlockingPost(Uri, Out.GetView(), EContentType::CbPackage, AcceptType);
 	}
 
-	FZenHttpRequest::Result FZenHttpRequest::PerformRpc(FStringView Uri, FCbObjectView Request, FCbPackage &OutResponse)
+	FZenHttpRequest::Result FZenHttpRequest::PerformRpc(FStringView Uri, FCbObjectView Request, FCbPackage& OutResponse)
 	{
-		FLargeMemoryWriter RequestBuffer;
-		
-		Request.CopyTo(RequestBuffer);
-		FCompositeBuffer Buffer(FSharedBuffer::MakeView(RequestBuffer.GetView()));
-		ReadDataView = &Buffer;
-
-		const uint32 ContentLength = RequestBuffer.TotalSize();
-
-		curl_easy_setopt(Curl, CURLOPT_POST, 1L);
-		curl_easy_setopt(Curl, CURLOPT_INFILESIZE, ContentLength);
-		curl_easy_setopt(Curl, CURLOPT_READDATA, this);
-		curl_easy_setopt(Curl, CURLOPT_READFUNCTION, &FZenHttpRequest::FStatics::StaticReadFn);
+		return ParseRpcResponse(
+			PerformBlockingPost(Uri, Request, EContentType::CbPackage), OutResponse);
+	}
 
-		AddHeader(TEXT("Content-Type"_SV), GetMimeType(EContentType::CbObject));
-		AddHeader(TEXT("Accept"_SV), GetMimeType(EContentType::CbPackage));
-		
-		Result RpcResult = PerformBlocking(Uri, RequestVerb::Post, ContentLength);
+	FZenHttpRequest::Result FZenHttpRequest::PerformRpc(FStringView Uri, const FCbPackage& Request, FCbPackage& OutResponse)
+	{
+		return ParseRpcResponse(
+			PerformBlockingPostPackage(Uri, Request, EContentType::CbPackage), OutResponse);
+	}
 
-		if (RpcResult != Result::Success || !IsSuccessCode(ResponseCode))
+	FZenHttpRequest::Result FZenHttpRequest::ParseRpcResponse(FZenHttpRequest::Result ResultFromPost, FCbPackage& OutResponse)
+	{
+		if (ResultFromPost != Result::Success || !IsSuccessCode(ResponseCode))
 		{
 			return Result::Failed;
 		}
@@ -360,26 +355,8 @@ namespace UE::Zen {
 		return Package;
 	}
 
-	FZenHttpRequest::Result FZenHttpRequest::PerformBlockingPost(FStringView Uri, FMemoryView Payload)
-	{
-		uint64 ContentLength = 0u;
-
-		curl_easy_setopt(Curl, CURLOPT_POST, 1L);
-		curl_easy_setopt(Curl, CURLOPT_INFILESIZE, Payload.GetSize());
-		curl_easy_setopt(Curl, CURLOPT_READDATA, this);
-		curl_easy_setopt(Curl, CURLOPT_READFUNCTION, &FZenHttpRequest::FStatics::StaticReadFn);
-
-		ContentLength = Payload.GetSize();
-
-		FCompositeBuffer Buffer(FSharedBuffer::MakeView(FMemoryView{ reinterpret_cast<const uint8*>(Payload.GetData()), Payload.GetSize() }));
-		ReadDataView = &Buffer;
-		
-		AddHeader(TEXT("Content-Type"_SV), GetMimeType(EContentType::Binary));
-
-		return PerformBlocking(Uri, RequestVerb::Post, ContentLength);
-	}
-
-	FZenHttpRequest::Result FZenHttpRequest::PerformBlockingPost(FStringView Uri, FMemoryView Payload, EContentType ContentType)
+	FZenHttpRequest::Result FZenHttpRequest::PerformBlockingPost(FStringView Uri, FMemoryView Payload,
+		EContentType ContentType, EContentType AcceptType)
 	{
 		uint64 ContentLength = 0u;
 
@@ -389,10 +366,14 @@ namespace UE::Zen {
 		curl_easy_setopt(Curl, CURLOPT_READFUNCTION, &FZenHttpRequest::FStatics::StaticReadFn);
 
 		AddHeader(TEXT("Content-Type"_SV), GetMimeType(ContentType));
+		if (AcceptType != EContentType::UnknownContentType)
+		{
+			AddHeader(TEXT("Accept"_SV), GetMimeType(EContentType::CbPackage));
+		}
 
 		ContentLength = Payload.GetSize();
 
-		FCompositeBuffer Buffer(FSharedBuffer::MakeView(FMemoryView{ reinterpret_cast<const uint8*>(Payload.GetData()), Payload.GetSize() }));
+		FCompositeBuffer Buffer(FSharedBuffer::MakeView(Payload));
 		ReadDataView = &Buffer;
 
 		return PerformBlocking(Uri, RequestVerb::Post, ContentLength);
diff --git a/Engine/Source/Developer/Zen/Internal/ZenServerHttp.h b/Engine/Source/Developer/Zen/Internal/ZenServerHttp.h
--- a/Engine/Source/Developer/Zen/Internal/ZenServerHttp.h
+++ b/Engine/Source/Developer/Zen/Internal/ZenServerHttp.h
@@ -167,12 +167,15 @@ public:
 		*/
 	ZEN_API Result PerformBlockingDelete(FStringView Uri);
 
-	ZEN_API Result PerformBlockingPost(FStringView Uri, FCbObjectView Obj);
-	ZEN_API Result PerformBlockingPostPackage(FStringView Uri, FCbPackage Package);
-	ZEN_API Result PerformBlockingPost(FStringView Uri, FMemoryView Payload);
-	ZEN_API Result PerformBlockingPost(FStringView Uri, FMemoryView Payload, EContentType ContentType);
+	ZEN_API Result PerformBlockingPostPackage(FStringView Uri, const FCbPackage& Package,
+		EContentType AcceptType = EContentType::UnknownContentType);
+	ZEN_API Result PerformBlockingPost(FStringView Uri, FCbObjectView Obj,
+		EContentType AcceptType = EContentType::UnknownContentType);
+	ZEN_API Result PerformBlockingPost(FStringView Uri, FMemoryView Payload,
+		EContentType ContentType = EContentType::Binary, EContentType AcceptType = EContentType::UnknownContentType);
 
 	ZEN_API Result PerformRpc(FStringView Uri, FCbObjectView Request, FCbPackage &OutResponse);
+	ZEN_API Result PerformRpc(FStringView Uri, const FCbPackage& Request, FCbPackage& OutResponse);
 
 	/** Returns the response buffer as a string. Note that is the request is performed
 		with an external buffer as target buffer this string will be empty.
@@ -243,6 +246,8 @@ private:
 		*/
 	Result PerformBlocking(FStringView Uri, RequestVerb Verb, uint32 ContentLength);
 
+	FZenHttpRequest::Result ParseRpcResponse(FZenHttpRequest::Result ResultFromPost, FCbPackage& OutResponse);
+
 	static FString GetAnsiBufferAsString(const TArray64<uint8>& Buffer);
 
 	struct FStatics;
