POST data size was not specified correctly.

Unknown POST sizes cause later libcurl versions to add an Transfer-Encoding header which is not compatible with the Content-Length header. An exert from the libcurl docs; "without chunked transfer, you must specify the size in the request. (Since 7.66.0, libcurl will automatically use chunked encoding for POSTs if the size is unknown.) When providing data with a callback, you must transmit it using chunked transfer-encoding or you must set the size of the data with the CURLOPT_POSTFIELDSIZE"

#rnx

diff --git a/Engine/Source/Developer/Zen/Internal/ZenServerHttp.cpp b/Engine/Source/Developer/Zen/Internal/ZenServerHttp.cpp
--- a/Engine/Source/Developer/Zen/Internal/ZenServerHttp.cpp
+++ b/Engine/Source/Developer/Zen/Internal/ZenServerHttp.cpp
@@ -361,7 +361,7 @@ namespace UE::Zen {
 		uint64 ContentLength = 0u;
 
 		curl_easy_setopt(Curl, CURLOPT_POST, 1L);
-		curl_easy_setopt(Curl, CURLOPT_INFILESIZE, Payload.GetSize());
+		curl_easy_setopt(Curl, CURLOPT_POSTFIELDSIZE, Payload.GetSize());
 		curl_easy_setopt(Curl, CURLOPT_READDATA, this);
 		curl_easy_setopt(Curl, CURLOPT_READFUNCTION, &FZenHttpRequest::FStatics::StaticReadFn);
 
