Use SIO_LOOPBACK_FAST_PATH on Windows when communicating with ZenServer.

diff --git a/Engine/Source/Developer/Zen/Internal/ZenServerHttp.cpp b/Engine/Source/Developer/Zen/Internal/ZenServerHttp.cpp
--- a/Engine/Source/Developer/Zen/Internal/ZenServerHttp.cpp
+++ b/Engine/Source/Developer/Zen/Internal/ZenServerHttp.cpp
@@ -12,6 +12,10 @@
 #	include "Windows/AllowWindowsPlatformTypes.h"
 #endif
 
+#if PLATFORM_WINDOWS
+#	include <mstcpip.h>
+#endif
+
 #if !defined(PLATFORM_CURL_INCLUDE)
 	#include "curl/curl.h"
 #endif
@@ -48,6 +52,9 @@ namespace UE::Zen {
 		static size_t StaticWriteHeaderFn(void* Ptr, size_t SizeInBlocks, size_t BlockSizeInBytes, void* UserData);
 		static size_t StaticWriteBodyFn(void* Ptr, size_t SizeInBlocks, size_t BlockSizeInBytes, void* UserData);
 		static size_t StaticSeekFn(void* UserData, curl_off_t Offset, int Origin);
+#if PLATFORM_WINDOWS
+		static int StaticSockoptFn(void* UserData, curl_socket_t CurlFd, curlsocktype Purpose);
+#endif //PLATFORM_WINDOWS
 	};
 
 	FZenHttpRequest::FZenHttpRequest(FStringView InDomain, bool bInLogErrors)
@@ -93,6 +100,9 @@ namespace UE::Zen {
 		// Rewind method, handle special error case where request need to rewind data stream
 		curl_easy_setopt(Curl, CURLOPT_SEEKDATA, this);
 		curl_easy_setopt(Curl, CURLOPT_SEEKFUNCTION, &FZenHttpRequest::FStatics::StaticSeekFn);
+#if PLATFORM_WINDOWS
+		curl_easy_setopt(Curl, CURLOPT_SOCKOPTFUNCTION, &FZenHttpRequest::FStatics::StaticSockoptFn);
+#endif //PLATFORM_WINDOWS
 		// Debug hooks
 #if UE_ZENDDC_HTTP_DEBUG
 		curl_easy_setopt(Curl, CURLOPT_DEBUGDATA, this);
@@ -779,6 +789,20 @@ namespace UE::Zen {
 		return CURL_SEEKFUNC_OK;
 	}
 
+#if PLATFORM_WINDOWS
+	int FZenHttpRequest::FStatics::StaticSockoptFn(void* UserData, curl_socket_t CurlFd, curlsocktype Purpose)
+	{
+		// On Windows, loopback connections can take advantage of a faster code path optionally with this flag.
+		// This must be used by both the client and server side, and is only effective in the absence of
+		// Windows Filtering Platform (WFP) callouts which can be installed by security software.
+		// https://docs.microsoft.com/en-us/windows/win32/winsock/sio-loopback-fast-path
+		int LoopbackOptionValue = 1;
+		DWORD OptionNumberOfBytesReturned = 0;
+		WSAIoctl(CurlFd, SIO_LOOPBACK_FAST_PATH, &LoopbackOptionValue, sizeof(LoopbackOptionValue), NULL, 0, &OptionNumberOfBytesReturned, 0, 0);
+		return CURL_SOCKOPT_OK;
+	}
+#endif //PLATFORM_WINDOWS
+
 	//////////////////////////////////////////////////////////////////////////
 
 	FZenHttpRequestPool::FZenHttpRequestPool(FStringView InServiceUrl, uint32 PoolEntryCount)
