Incorrect variable name in POSIX #if/def block.

#rnx

diff --git a/Engine/Source/Developer/Zen/Internal/ZenServerInterface.cpp b/Engine/Source/Developer/Zen/Internal/ZenServerInterface.cpp
--- a/Engine/Source/Developer/Zen/Internal/ZenServerInterface.cpp
+++ b/Engine/Source/Developer/Zen/Internal/ZenServerInterface.cpp
@@ -829,7 +829,7 @@ FZenServiceInstance::AutoLaunch(const FServiceAutoLaunchSettings& InSettings, FS
 			void* PipeWriteChild = nullptr;
 			void* PipeReadChild = nullptr;
 			Proc = FPlatformProcess::CreateProc(
-				*MainFilePath,
+				*ExecutablePath,
 				*Parms,
 				bLaunchDetached,
 				bLaunchHidden,
