Implementation of the set side of Zen's NamedEvent for Linux and Mac.

#rnx

diff --git a/Engine/Source/Developer/Zen/Internal/ZenServerInterface.cpp b/Engine/Source/Developer/Zen/Internal/ZenServerInterface.cpp
--- a/Engine/Source/Developer/Zen/Internal/ZenServerInterface.cpp
+++ b/Engine/Source/Developer/Zen/Internal/ZenServerInterface.cpp
@@ -34,6 +34,10 @@
 #	include "Windows/HideWindowsPlatformTypes.h"
 #endif
 
+#if PLATFORM_UNIX || PLATFORM_MAC
+#	include <sys/sem.h>
+#endif
+
 #define ALLOW_SETTINGS_OVERRIDE_FROM_COMMANDLINE			(UE_SERVER || !(UE_BUILD_SHIPPING))
 
 namespace UE::Zen
@@ -453,6 +457,24 @@ RequestZenShutdownOnPort(uint16 Port)
 		ON_SCOPE_EXIT{ CloseHandle(Handle); };
 		SetEvent(Handle);
 	}
+#elif PLATFORM_UNIX || PLATFORM_MAC
+	TAnsiStringBuilder<64> EventPath;
+	EventPath << "/tmp/Zen_" << Port << "_Shutdown";
+
+	key_t IpcKey = ftok(EventPath.ToString(), 1);
+	if (IpcKey < 0)
+	{
+		return;
+	}
+
+	int Semaphore = semget(IpcKey, 1, 0600);
+	if (Semaphore < 0)
+	{
+		return;
+	}
+
+	semctl(Semaphore, 0, SETVAL, 0);
+	semctl(Semaphore, 0, IPC_RMID);
 #else
 	static_assert(false, "Missing implementation for Zen named shutdown events");
 #endif
