Fixed unresolved symbol errors. An inlined and exported function is very portable. MSVC respects the export, while Apple Clang favours the inline and fails on symbols the inlined body references that aren't exported.

#rnx

diff --git a/Engine/Source/Developer/Zen/Internal/ZenServerInterface.cpp b/Engine/Source/Developer/Zen/Internal/ZenServerInterface.cpp
--- a/Engine/Source/Developer/Zen/Internal/ZenServerInterface.cpp
+++ b/Engine/Source/Developer/Zen/Internal/ZenServerInterface.cpp
@@ -636,6 +636,12 @@ FZenServiceInstance::IsServiceReady()
 	return false;
 }
 
+uint16
+FZenServiceInstance::GetAutoLaunchedPort()
+{
+	return AutoLaunchedPort;
+}
+
 void
 FZenServiceInstance::Initialize()
 {
diff --git a/Engine/Source/Developer/Zen/Public/ZenServerInterface.h b/Engine/Source/Developer/Zen/Public/ZenServerInterface.h
--- a/Engine/Source/Developer/Zen/Public/ZenServerInterface.h
+++ b/Engine/Source/Developer/Zen/Public/ZenServerInterface.h
@@ -128,7 +128,7 @@ public:
 	UE_API bool IsServiceReady();
 	UE_API bool IsServiceRunningLocally() const { return bIsRunningLocally; }
 
-	static UE_API uint16 GetAutoLaunchedPort() { return AutoLaunchedPort; }
+	static UE_API uint16 GetAutoLaunchedPort();
 
 private:
 
