Moved lock file existence checks into a function. Implemented lock testing and lock file reading for POSIX-based platforms.
#rnx

diff --git a/Engine/Source/Developer/Zen/Internal/ZenServerInterface.cpp b/Engine/Source/Developer/Zen/Internal/ZenServerInterface.cpp
--- a/Engine/Source/Developer/Zen/Internal/ZenServerInterface.cpp
+++ b/Engine/Source/Developer/Zen/Internal/ZenServerInterface.cpp
@@ -430,20 +430,85 @@ ReadCbLockFile(FStringView FileName, FCbObject& OutLockObject)
 		}
 	}
 	return false;
-#else
-	// Generic lock reading path
-	if (TUniquePtr<FArchive> Ar{FileManager.CreateFileReader(*FileName, FILEREAD_AllowWrite | FILEREAD_Silent)})
+#elif PLATFORM_UNIX || PLATFORM_MAC
+	TAnsiStringBuilder<256> LockFilePath;
+	LockFilePath << FileName;
+	int32 Fd = open(LockFilePath.ToString(), O_RDONLY);
+	if (Fd < 0)
+	{
+		return false;
+	}
+
+	// If we can claim the lock then it's an orphaned lock file and should be
+	// ignored. Not ideal as there's a period of time when the lock can be
+	// held unncessarily.
+	int32 LockRet = flock(Fd, LOCK_EX | LOCK_NB);
+	if (LockRet >= 0)
+	{
+		unlink(LockFilePath.ToString());
+		flock(Fd, LOCK_UN);
+		close(Fd);
+		return false;
+	}
+
+	if (errno != EWOULDBLOCK && errno != EAGAIN)
 	{
-		*Ar << OutLockObject;
-		FMemoryView View;
-		if (Ar.Close() &&
-			OutLockObject.TryGetView(View) &&
-			ValidateCompactBinary(View, ECbValidateMode::Default) == ECbValidateError::None)
+		return false;
+	}
+
+	struct stat Stat;
+	fstat(Fd, &Stat);
+	uint64 FileSize = uint64(Stat.st_size);
+
+	bool bSuccess = false;
+	FUniqueBuffer FileBytes = FUniqueBuffer::Alloc(FileSize);
+	if (read(Fd, FileBytes.GetData(), FileSize) == FileSize)
+	{
+		if (ValidateCompactBinary(FileBytes, ECbValidateMode::Default) == ECbValidateError::None)
 		{
-			return true;
+			OutLockObject = FCbObject(FileBytes.MoveToShared());
+			bSuccess = true;
 		}
 	}
-	return false;
+
+	close(Fd);
+	return bSuccess;
+#endif
+}
+
+static bool
+IsLockFileLocked(const TCHAR* FileName, bool bAttemptCleanUp=false)
+{
+#if PLATFORM_WINDOWS
+	if (bAttemptCleanUp)
+	{
+		IFileManager::Get().Delete(FileName, false, false, true);
+	}
+	return IFileManager::Get().FileExists(FileName);
+#elif PLATFORM_UNIX || PLATFORM_MAC
+	TAnsiStringBuilder<256> LockFilePath;
+	LockFilePath << FileName;
+	int32 Fd = open(LockFilePath.ToString(), O_RDONLY);
+	if (Fd < 0)
+	{
+		return false;
+	}
+
+	int32 LockRet = flock(Fd, LOCK_EX | LOCK_NB);
+	if (LockRet < 0)
+	{
+		close(Fd);
+		return errno == EWOULDBLOCK || errno == EAGAIN;
+	}
+
+	// Consider the lock file as orphaned if we we managed to claim the lock for
+	// it. Might as well delete it while we own it.
+	unlink(LockFilePath.ToString());
+
+	flock(Fd, LOCK_UN);
+	close(Fd);
+
+	return true;
 #endif
 }
 
@@ -484,7 +549,7 @@ static bool
 WaitForZenShutdown(const TCHAR* LockFilePath, double MaximumWaitDurationSeconds)
 {
 	uint64 ZenShutdownWaitStartTime = FPlatformTime::Cycles64();
-	while (IFileManager::Get().FileExists(LockFilePath))
+	while (IsLockFileLocked(LockFilePath))
 	{
 		double ZenShutdownWaitDuration = FPlatformTime::ToSeconds64(FPlatformTime::Cycles64() - ZenShutdownWaitStartTime);
 		if (ZenShutdownWaitDuration < MaximumWaitDurationSeconds)
@@ -706,7 +771,7 @@ FZenServiceInstance::ConditionalUpdateLocalInstall()
 				}
 			}
 
-			if (FileManager.FileExists(*LockFilePath))
+			if (IsLockFileLocked(*LockFilePath))
 			{
 				PromptUserToStopRunningServerInstance(InstallFilePath);
 			}
@@ -736,11 +801,10 @@ FZenServiceInstance::AutoLaunch(const FServiceAutoLaunchSettings& InSettings, FS
 	IFileManager& FileManager = IFileManager::Get();
 	const FString LockFilePath = FPaths::Combine(InSettings.DataPath, TEXT(".lock"));
 	const FString CmdLineFilePath = FPaths::Combine(InSettings.DataPath, TEXT(".cmdline"));
-	FileManager.Delete(*LockFilePath, false, false, true);
 
 	bool bReUsingExistingInstance = false;
 
-	if (FileManager.FileExists(*LockFilePath))
+	if (IsLockFileLocked(*LockFilePath, true))
 	{
 		// If an instance is running with this data path, check if we can use it and what port it is on
 		uint16 CurrentPort = 0;
@@ -775,7 +839,7 @@ FZenServiceInstance::AutoLaunch(const FServiceAutoLaunchSettings& InSettings, FS
 	}
 
 
-	bool bProcessIsLive = FileManager.FileExists(*LockFilePath);
+	bool bProcessIsLive = IsLockFileLocked(*LockFilePath);
 
 	// When limiting process lifetime, always re-launch to add sponsor process IDs.
 	// When not limiting process lifetime, only launch if the process is not already live.
@@ -918,7 +982,7 @@ FZenServiceInstance::AutoLaunch(const FServiceAutoLaunchSettings& InSettings, FS
 			{
 				if (DurationPhase == EWaitDurationPhase::Short)
 				{
-					if (!FileManager.FileExists(*LockFilePath))
+					if (!IsLockFileLocked(*LockFilePath))
 					{
 						if (FApp::IsUnattended())
 						{
