Fix excessive thread spawning caused by name resolving of "localhost" at high frequency by the Zen stats code.
Re-use Curl request when gathering stats.  Rate limit it to no more than 2/second.  Still keeping it synchronous and blocking for now.
Avoid using "localhost" and prefer the IPV6 loopback address of [::1] to avoid name resolution penalty for localhost.
Keep DNS cache in Curl permanently for a given request.

diff --git a/Engine/Source/Developer/Zen/Internal/ZenServerHttp.cpp b/Engine/Source/Developer/Zen/Internal/ZenServerHttp.cpp
--- a/Engine/Source/Developer/Zen/Internal/ZenServerHttp.cpp
+++ b/Engine/Source/Developer/Zen/Internal/ZenServerHttp.cpp
@@ -93,6 +93,7 @@ namespace UE::Zen {
 		curl_easy_setopt(Curl, CURLOPT_CONNECTTIMEOUT, 30L);
 		curl_easy_setopt(Curl, CURLOPT_EXPECT_100_TIMEOUT_MS, 0);
 		curl_easy_setopt(Curl, CURLOPT_NOSIGNAL, 1L);
+		curl_easy_setopt(Curl, CURLOPT_DNS_CACHE_TIMEOUT, -1L); // Don't re-resolve names mid-session
 		curl_easy_setopt(Curl, CURLOPT_BUFFERSIZE, 256 * 1024L);
 		//curl_easy_setopt(Curl, CURLOPT_UPLOAD_BUFFERSIZE, 256 * 1024L);
 		// Response functions
diff --git a/Engine/Source/Developer/Zen/Internal/ZenServerHttp.h b/Engine/Source/Developer/Zen/Internal/ZenServerHttp.h
--- a/Engine/Source/Developer/Zen/Internal/ZenServerHttp.h
+++ b/Engine/Source/Developer/Zen/Internal/ZenServerHttp.h
@@ -2,7 +2,7 @@
 
 #pragma once
 
-#include "ZenServerInterface.h"
+#include "ZenGlobals.h"
 
 #include "Containers/StaticArray.h"
 #include "Containers/UnrealString.h"
diff --git a/Engine/Source/Developer/Zen/Internal/ZenServerInterface.cpp b/Engine/Source/Developer/Zen/Internal/ZenServerInterface.cpp
--- a/Engine/Source/Developer/Zen/Internal/ZenServerInterface.cpp
+++ b/Engine/Source/Developer/Zen/Internal/ZenServerInterface.cpp
@@ -3,7 +3,6 @@
 
 
 #include "ZenBackendUtils.h"
-#include "ZenServerHttp.h"
 #include "ZenSerialization.h"
 
 #include "Dom/JsonValue.h"
@@ -187,6 +186,11 @@ IsLocalHost(const FString& Host)
 		return true;
 	}
 
+	if (Host.Compare(FString(TEXT("[::1]"))) == 0)
+	{
+		return true;
+	}
+
 	ISocketSubsystem& SocketSubsystem = *ISocketSubsystem::Get();
 
 	const TSharedPtr<FInternetAddr> Addr = SocketSubsystem.GetAddressFromString(Host);
@@ -351,7 +355,7 @@ FServiceSettings::TryApplyAutoLaunchOverride()
 	{
 		SettingsVariant.Emplace<FServiceConnectSettings>();
 		FServiceConnectSettings& ConnectExistingSettings = SettingsVariant.Get<FServiceConnectSettings>();
-		ConnectExistingSettings.HostName = TEXT("localhost");
+		ConnectExistingSettings.HostName = TEXT("[::1]");
 		ConnectExistingSettings.Port = 1337;
 		return true;
 	}
@@ -943,7 +947,7 @@ FZenServiceInstance::AutoLaunch(const FServiceAutoLaunchSettings& InSettings, FS
 	}
 
 
-	OutHostName = TEXT("localhost");
+	OutHostName = TEXT("[::1]");
 	// Default to assuming that we get to run on the port we want
 	OutPort = DesiredPort;
 
@@ -1041,9 +1045,23 @@ FZenServiceInstance::AutoLaunch(const FServiceAutoLaunchSettings& InSettings, FS
 bool 
 FZenServiceInstance::GetStats( FZenStats& stats ) const
 {
+	check(IsInGameThread());
+
+	const uint64 CurrentCycles = FPlatformTime::Cycles64();
+	if (StatsRequest.IsSet() && FPlatformTime::ToSeconds(CurrentCycles - LastStatsTime) < 0.5)
+	{
+		stats = LastStats;
+		return true;
+	}
+
 	TStringBuilder<128> ZenDomain;
 	ZenDomain << HostName << TEXT(":") << Port;
-	UE::Zen::FZenHttpRequest Request(ZenDomain.ToString(), false);
+	if (!StatsRequest.IsSet())
+	{
+		StatsRequest.Emplace(ZenDomain.ToString(), false);
+	}
+	UE::Zen::FZenHttpRequest& Request = StatsRequest.GetValue();
+	Request.Reset();
 
 	TArray64<uint8> GetBuffer;
 	FZenHttpRequest::Result Result = Request.PerformBlockingDownload(TEXT("/stats/z$"_SV), &GetBuffer, Zen::EContentType::CbObject);
@@ -1129,6 +1147,8 @@ FZenServiceInstance::GetStats( FZenStats& stats ) const
 		CASSizeStats.Large = CASSizeObjectView["large"].AsInt64();
 		CASSizeStats.Total = CASSizeObjectView["total"].AsInt64();
 
+		LastStats = stats;
+		LastStatsTime = CurrentCycles;
 		return true;
 	}
 
diff --git a/Engine/Source/Developer/Zen/Public/ZenGlobals.h b/Engine/Source/Developer/Zen/Public/ZenGlobals.h
new file mode 100644
--- /dev/null
+++ b/Engine/Source/Developer/Zen/Public/ZenGlobals.h
@@ -0,0 +1,13 @@
+// Copyright Epic Games, Inc. All Rights Reserved.
+
+#pragma once
+
+#include "HAL/Platform.h"
+
+#ifndef UE_WITH_ZEN
+#	if PLATFORM_WINDOWS
+#		define UE_WITH_ZEN 1
+#	else
+#		define UE_WITH_ZEN 0
+#	endif
+#endif
diff --git a/Engine/Source/Developer/Zen/Public/ZenServerInterface.h b/Engine/Source/Developer/Zen/Public/ZenServerInterface.h
--- a/Engine/Source/Developer/Zen/Public/ZenServerInterface.h
+++ b/Engine/Source/Developer/Zen/Public/ZenServerInterface.h
@@ -6,18 +6,13 @@
 #include "Containers/StringView.h"
 #include "Containers/UnrealString.h"
 #include "Dom/JsonObject.h"
+#include "Misc/Optional.h"
 #include "Misc/TVariant.h"
 #include "Policies/PrettyJsonPrintPolicy.h"
 #include "Serialization/JsonWriter.h"
 #include "Templates/UniquePtr.h"
-
-#ifndef UE_WITH_ZEN
-#	if PLATFORM_WINDOWS
-#		define UE_WITH_ZEN 1
-#	else
-#		define UE_WITH_ZEN 0
-#	endif
-#endif
+#include "ZenGlobals.h"
+#include "ZenServerHttp.h"
 
 #define UE_API ZEN_API
 
@@ -137,6 +132,10 @@ private:
 	FString ConditionalUpdateLocalInstall();
 	static bool AutoLaunch(const FServiceAutoLaunchSettings& InSettings, FString&& ExecutablePath, FString& OutHostName, uint16& OutPort);
 
+	mutable TOptional<FZenHttpRequest> StatsRequest;
+	mutable FZenStats LastStats;
+	mutable uint64 LastStatsTime = 0;
+
 	FServiceSettings Settings;
 	FString URL;
 	FString HostName;
