Fix NonUnity compile error in ZenServerHttp.h.

diff --git a/Engine/Source/Developer/Zen/Internal/ZenServerHttp.h b/Engine/Source/Developer/Zen/Internal/ZenServerHttp.h
--- a/Engine/Source/Developer/Zen/Internal/ZenServerHttp.h
+++ b/Engine/Source/Developer/Zen/Internal/ZenServerHttp.h
@@ -5,6 +5,7 @@
 #include "ZenGlobals.h"
 
 #include "Containers/StaticArray.h"
+#include "Containers/StringView.h"
 #include "Containers/UnrealString.h"
 #include "Memory/MemoryFwd.h"
 
