Fix bug when parsing URL for Zen containing colons (eg: ::1).

diff --git a/Engine/Source/Developer/Zen/Internal/ZenServerInterface.cpp b/Engine/Source/Developer/Zen/Internal/ZenServerInterface.cpp
--- a/Engine/Source/Developer/Zen/Internal/ZenServerInterface.cpp
+++ b/Engine/Source/Developer/Zen/Internal/ZenServerInterface.cpp
@@ -308,7 +308,7 @@ FServiceSettings::ReadFromURL(FStringView InstanceURL)
 	}
 
 	int32 PortDelimIndex = INDEX_NONE;
-	InstanceURL.FindChar(TEXT(':'), PortDelimIndex);
+	InstanceURL.FindLastChar(TEXT(':'), PortDelimIndex);
 	if (PortDelimIndex != INDEX_NONE)
 	{
 		ConnectExistingSettings.HostName = InstanceURL.Left(PortDelimIndex);
