Missing include.

diff --git a/Engine/Source/Developer/Zen/Internal/ZenServerInterface.cpp b/Engine/Source/Developer/Zen/Internal/ZenServerInterface.cpp
--- a/Engine/Source/Developer/Zen/Internal/ZenServerInterface.cpp
+++ b/Engine/Source/Developer/Zen/Internal/ZenServerInterface.cpp
@@ -34,6 +34,7 @@
 #endif
 
 #if PLATFORM_UNIX || PLATFORM_MAC
+#	include <sys/file.h>
 #	include <sys/sem.h>
 #endif
 
