TEXT(""_SV) changed to TEXTVIEW("") for UTF-8 support.

diff --git a/Engine/Source/Developer/DerivedDataCache/Private/FileSystemCacheStore.cpp b/Engine/Source/Developer/DerivedDataCache/Private/FileSystemCacheStore.cpp
--- a/Engine/Source/Developer/DerivedDataCache/Private/FileSystemCacheStore.cpp
+++ b/Engine/Source/Developer/DerivedDataCache/Private/FileSystemCacheStore.cpp
@@ -94,7 +94,7 @@ void BuildPathForCachePackage(const FCacheKey& CacheKey, FStringBuilderBase& Pat
 	const FIoHash::ByteArray& Bytes = CacheKey.Hash.GetBytes();
 	Path.Appendf(TEXT("%s/%hs/%02x/%02x/"), GBucketsDirectoryName, CacheKey.Bucket.ToCString(), Bytes[0], Bytes[1]);
 	UE::String::BytesToHexLower(MakeArrayView(Bytes).RightChop(2), Path);
-	Path << TEXT(".udd"_SV);
+	Path << TEXTVIEW(".udd");
 }
 
 void BuildPathForCacheContent(const FIoHash& RawHash, FStringBuilderBase& Path)
@@ -102,7 +102,7 @@ void BuildPathForCacheContent(const FIoHash& RawHash, FStringBuilderBase& Path)
 	const FIoHash::ByteArray& Bytes = RawHash.GetBytes();
 	Path.Appendf(TEXT("%s/%02x/%02x/"), GContentDirectoryName, Bytes[0], Bytes[1]);
 	UE::String::BytesToHexLower(MakeArrayView(Bytes).RightChop(2), Path);
-	Path << TEXT(".udd"_SV);
+	Path << TEXTVIEW(".udd");
 }
 
 ///////////////////////////////////////////////////////////////////////////////////////////////////////////////////////
diff --git a/Engine/Source/Developer/DerivedDataCache/Private/ZenCacheStore.cpp b/Engine/Source/Developer/DerivedDataCache/Private/ZenCacheStore.cpp
--- a/Engine/Source/Developer/DerivedDataCache/Private/ZenCacheStore.cpp
+++ b/Engine/Source/Developer/DerivedDataCache/Private/ZenCacheStore.cpp
@@ -850,7 +850,7 @@ void FZenCacheStore::Get(
 		{
 			LLM_SCOPE_BYTAG(UntaggedDDCResult);
 			Zen::FZenScopedRequestPtr Request(RequestPool.Get());
-			HttpResult = Request->PerformRpc(TEXT("/z$/$rpc"_SV), BatchRequest.Save().AsObject(), BatchResponse);
+			HttpResult = Request->PerformRpc(TEXTVIEW("/z$/$rpc"), BatchRequest.Save().AsObject(), BatchResponse);
 		}
 
 		int32 RequestIndex = 0;
@@ -1225,7 +1225,7 @@ void FZenCacheStore::GetChunks(
 		{
 			LLM_SCOPE_BYTAG(UntaggedDDCResult);
 			Zen::FZenScopedRequestPtr Request(RequestPool.Get());
-			HttpResult = Request->PerformRpc(TEXT("/z$/$rpc"_SV), BatchRequest.Save().AsObject(), BatchResponse);
+			HttpResult = Request->PerformRpc(TEXTVIEW("/z$/$rpc"), BatchRequest.Save().AsObject(), BatchResponse);
 		}
 
 		int32 RequestIndex = 0;
diff --git a/Engine/Source/Developer/Zen/Internal/ZenServerHttp.cpp b/Engine/Source/Developer/Zen/Internal/ZenServerHttp.cpp
--- a/Engine/Source/Developer/Zen/Internal/ZenServerHttp.cpp
+++ b/Engine/Source/Developer/Zen/Internal/ZenServerHttp.cpp
@@ -17,7 +17,7 @@
 #endif
 
 #if !defined(PLATFORM_CURL_INCLUDE)
-	#include "curl/curl.h"
+#include "curl/curl.h"
 #endif
 
 #if PLATFORM_WINDOWS || PLATFORM_HOLOLENS
@@ -123,7 +123,7 @@ namespace UE::Zen {
 	void FZenHttpRequest::AddHeader(FStringView Header, FStringView Value)
 	{
 		TStringBuilder<128> Sb;
-		Sb << Header << TEXT(": "_SV) << Value;
+		Sb << Header << TEXTVIEW(": ") << Value;
 		Headers.Emplace(Sb.ToString());
 	}
 
@@ -137,7 +137,7 @@ namespace UE::Zen {
 		curl_easy_setopt(Curl, CURLOPT_READDATA, this);
 		curl_easy_setopt(Curl, CURLOPT_READFUNCTION, &FZenHttpRequest::FStatics::StaticReadFn);
 
-		AddHeader(TEXT("Content-Type"_SV), GetMimeType(ContentType));
+		AddHeader(TEXTVIEW("Content-Type"), GetMimeType(ContentType));
 
 		ReadDataView = &Buffer;
 
@@ -274,12 +274,12 @@ namespace UE::Zen {
 		return PerformBlockingPost(Uri, Out.GetView(), EContentType::CbPackage, AcceptType);
 	}
 
-	FZenHttpRequest::Result FZenHttpRequest::PerformRpc(FStringView Uri, FCbObjectView Request, FCbPackage& OutResponse)
+	FZenHttpRequest::Result FZenHttpRequest::PerformRpc(FStringView Uri, FCbObjectView Request, FCbPackage &OutResponse)
 	{
 		return ParseRpcResponse(
 			PerformBlockingPost(Uri, Request, EContentType::CbPackage), OutResponse);
 	}
-
+		
 	FZenHttpRequest::Result FZenHttpRequest::PerformRpc(FStringView Uri, const FCbPackage& Request, FCbPackage& OutResponse)
 	{
 		return ParseRpcResponse(
@@ -379,11 +379,11 @@ namespace UE::Zen {
 		curl_easy_setopt(Curl, CURLOPT_READDATA, this);
 		curl_easy_setopt(Curl, CURLOPT_READFUNCTION, &FZenHttpRequest::FStatics::StaticReadFn);
 
-		AddHeader(TEXT("Content-Type"_SV), GetMimeType(ContentType));
+		AddHeader(TEXTVIEW("Content-Type"), GetMimeType(ContentType));
 		if (AcceptType != EContentType::UnknownContentType)
 		{
-			AddHeader(TEXT("Accept"_SV), GetMimeType(EContentType::CbPackage));
-		}
+			AddHeader(TEXTVIEW("Accept"), GetMimeType(EContentType::CbPackage));
+	}
 
 		ContentLength = Payload.GetSize();
 
@@ -398,7 +398,7 @@ namespace UE::Zen {
 		curl_easy_setopt(Curl, CURLOPT_HTTPGET, 1L);
 		WriteDataBufferPtr = Buffer;
 		
-		AddHeader(TEXT("Accept"_SV), GetMimeType(AcceptType));
+		AddHeader(TEXTVIEW("Accept"), GetMimeType(AcceptType));
 
 		return PerformBlocking(Uri, RequestVerb::Get, 0u);
 	}
@@ -408,7 +408,7 @@ namespace UE::Zen {
 		curl_easy_setopt(Curl, CURLOPT_HTTPGET, 1L);
 		OutPackage.Reset();
 
-		AddHeader(TEXT("Accept"_SV), GetMimeType(EContentType::CbPackage));
+		AddHeader(TEXTVIEW("Accept"), GetMimeType(EContentType::CbPackage));
 
 		// TODO: When PackageBytes can be written in segments directly, set the WritePtr to the OutPackage and use that
 		TArray64<uint8> PackageBytes;
@@ -426,7 +426,7 @@ namespace UE::Zen {
 	{
 		curl_easy_setopt(Curl, CURLOPT_NOBODY, 1L);
 		
-		AddHeader(TEXT("Accept"_SV), GetMimeType(AcceptType));
+		AddHeader(TEXTVIEW("Accept"), GetMimeType(AcceptType));
 
 		return PerformBlocking(Uri, RequestVerb::Head, 0u);
 	}
diff --git a/Engine/Source/Developer/Zen/Internal/ZenServerHttp.h b/Engine/Source/Developer/Zen/Internal/ZenServerHttp.h
--- a/Engine/Source/Developer/Zen/Internal/ZenServerHttp.h
+++ b/Engine/Source/Developer/Zen/Internal/ZenServerHttp.h
@@ -57,23 +57,23 @@ inline FStringView GetMimeType(EContentType Type)
 	switch (Type)
 	{
 		case EContentType::Binary:
-			return TEXT("application/octet-stream"_SV);
+			return TEXTVIEW("application/octet-stream");
 		case EContentType::Text:
-			return TEXT("text/plain"_SV);
+			return TEXTVIEW("text/plain");
 		case EContentType::JSON:
-			return TEXT("application/json"_SV);
+			return TEXTVIEW("application/json");
 		case EContentType::CbObject:
-			return TEXT("application/x-ue-cb"_SV);
+			return TEXTVIEW("application/x-ue-cb");
 		case EContentType::CbPackage:
-			return TEXT("application/x-ue-cbpkg"_SV);
+			return TEXTVIEW("application/x-ue-cbpkg");
 		case EContentType::YAML:
-			return TEXT("text/yaml"_SV);
+			return TEXTVIEW("text/yaml");
 		case EContentType::CbPackageOffer:
-			return TEXT("application/x-ue-offer"_SV);
+			return TEXTVIEW("application/x-ue-offer");
 		case EContentType::CompressedBinary:
-			return TEXT("application/x-ue-comp"_SV);
+			return TEXTVIEW("application/x-ue-comp");
 		default:
-			return TEXT("unknown"_SV);
+			return TEXTVIEW("unknown");
 	}
 }
 
diff --git a/Engine/Source/Developer/Zen/Internal/ZenServerInterface.cpp b/Engine/Source/Developer/Zen/Internal/ZenServerInterface.cpp
--- a/Engine/Source/Developer/Zen/Internal/ZenServerInterface.cpp
+++ b/Engine/Source/Developer/Zen/Internal/ZenServerInterface.cpp
@@ -411,7 +411,7 @@ ReadCbLockFile(FStringView FileName, FCbObject& OutLockObject)
 	}
 	if (FullFileNameBuilder.Len() >= MAX_PATH)
 	{
-		FullFileNameBuilder.Prepend(TEXT("\\\\?\\"_SV));
+		FullFileNameBuilder.Prepend(TEXTVIEW("\\\\?\\"));
 	}
 	HANDLE Handle = CreateFileW(FullFileNameBuilder.ToString(), Access, WinFlags, NULL, Create, FILE_ATTRIBUTE_NORMAL, NULL);
 	if (Handle != INVALID_HANDLE_VALUE)
@@ -691,7 +691,7 @@ FZenServiceInstance::IsServiceReady()
 		TStringBuilder<128> ZenDomain;
 		ZenDomain << HostName << TEXT(":") << Port;
 		Zen::FZenHttpRequest Request(ZenDomain.ToString(), false);
-		Zen::FZenHttpRequest::Result Result = Request.PerformBlockingDownload(TEXT("health/ready"_SV), nullptr, Zen::EContentType::Text);
+		Zen::FZenHttpRequest::Result Result = Request.PerformBlockingDownload(TEXTVIEW("health/ready"), nullptr, Zen::EContentType::Text);
 		
 		if (Result == Zen::FZenHttpRequest::Result::Success && Zen::IsSuccessCode(Request.GetResponseCode()))
 		{
@@ -1065,7 +1065,7 @@ FZenServiceInstance::GetStats( FZenStats& stats ) const
 	Request.Reset();
 
 	TArray64<uint8> GetBuffer;
-	FZenHttpRequest::Result Result = Request.PerformBlockingDownload(TEXT("/stats/z$"_SV), &GetBuffer, Zen::EContentType::CbObject);
+	FZenHttpRequest::Result Result = Request.PerformBlockingDownload(TEXTVIEW("/stats/z$"), &GetBuffer, Zen::EContentType::CbObject);
 
 	if (Result == Zen::FZenHttpRequest::Result::Success && Request.GetResponseCode() == 200)
 	{
diff --git a/Engine/Source/Runtime/Core/Private/Windows/WindowsPlatformFile.cpp b/Engine/Source/Runtime/Core/Private/Windows/WindowsPlatformFile.cpp
--- a/Engine/Source/Runtime/Core/Private/Windows/WindowsPlatformFile.cpp
+++ b/Engine/Source/Runtime/Core/Private/Windows/WindowsPlatformFile.cpp
@@ -934,14 +934,14 @@ class CORE_API FWindowsPlatformFile : public IPhysicalPlatformFile
 		}
 
 		// Remove duplicate slashes
-		const bool bIsUNCPath = Path.ToView().StartsWith(TEXT("//"_SV));
+		const bool bIsUNCPath = Path.ToView().StartsWith(TEXTVIEW("//"));
 		
 		FPathViews::RemoveDuplicateSlashes(Path);
 
 		if (bIsUNCPath)
 		{
 			// Keep // at the beginning.  If There are more than two / at the beginning, replace them with just //.
-			Path.Prepend(TEXT("/"_SV));
+			Path.Prepend(TEXTVIEW("/"));
 		}
 
 		// We now have a canonical, strict-valid, absolute Unreal Path.  Convert it to a Windows Path.
@@ -958,11 +958,11 @@ class CORE_API FWindowsPlatformFile : public IPhysicalPlatformFile
 		{
 			if (bIsUNCPath)
 			{
-				Path.ReplaceAt(0, 1, TEXT("\\\\?\\UNC"_SV));
+				Path.ReplaceAt(0, 1, TEXTVIEW("\\\\?\\UNC"));
 			}
 			else
 			{
-				Path.Prepend(TEXT("\\\\?\\"_SV));
+				Path.Prepend(TEXTVIEW("\\\\?\\"));
 			}
 		}
 	}
@@ -1461,7 +1461,7 @@ public:
 	// Outline to reduce stack space usage since IterateDirectoryCommon might be recursive
 	FORCENOINLINE static HANDLE FindFirstFileWithWildcard(const TCHAR* Directory, WIN32_FIND_DATAW& OutData)
 	{
-		return FindFirstFileW(*(FNormalizedFilename(Directory, TEXT("*.*"_SV))), &OutData);
+		return FindFirstFileW(*(FNormalizedFilename(Directory, TEXTVIEW("*.*"))), &OutData);
 	}
 
 	bool IterateDirectoryCommon(const TCHAR* Directory, const TFunctionRef<bool(const WIN32_FIND_DATAW&)>& Visitor)
diff --git a/Engine/Source/Runtime/CoreUObject/Private/Misc/PackageName.cpp b/Engine/Source/Runtime/CoreUObject/Private/Misc/PackageName.cpp
--- a/Engine/Source/Runtime/CoreUObject/Private/Misc/PackageName.cpp
+++ b/Engine/Source/Runtime/CoreUObject/Private/Misc/PackageName.cpp
@@ -2353,7 +2353,7 @@ bool FPackageName::IsLocalizedPackage(FStringView InPackageName)
 	// Are we part of the L10N folder?
 	FStringView Remaining(CurChar, UE_PTRDIFF_TO_INT32(EndChar - CurChar));
 	// Is "L10N" or StartsWith "L10N/" 
-	return Remaining.StartsWith(TEXT("L10N"_SV), ESearchCase::IgnoreCase) && (Remaining.Len() == 4 || Remaining[4] == '/');
+	return Remaining.StartsWith(TEXTVIEW("L10N"), ESearchCase::IgnoreCase) && (Remaining.Len() == 4 || Remaining[4] == '/');
 }
 
 FString FPackageName::FormatErrorAsString(FStringView InPath, EErrorCode ErrorCode)
diff --git a/Engine/Source/Runtime/CoreUObject/Private/Misc/PackagePath.cpp b/Engine/Source/Runtime/CoreUObject/Private/Misc/PackagePath.cpp
--- a/Engine/Source/Runtime/CoreUObject/Private/Misc/PackagePath.cpp
+++ b/Engine/Source/Runtime/CoreUObject/Private/Misc/PackagePath.cpp
@@ -137,31 +137,31 @@ EPackageExtension FPackagePath::ParseExtension(FStringView Filename, int32* OutE
 	{
 		return EPackageExtension::Unspecified;
 	}
-	if (Extension.Equals(TEXT(".uasset"_SV), ESearchCase::IgnoreCase))
+	if (Extension.Equals(TEXTVIEW(".uasset"), ESearchCase::IgnoreCase))
 	{
 		return EPackageExtension::Asset;
 	}
-	if (Extension.Equals(TEXT(".umap"_SV), ESearchCase::IgnoreCase))
+	if (Extension.Equals(TEXTVIEW(".umap"), ESearchCase::IgnoreCase))
 	{
 		return EPackageExtension::Map;
 	}
-	if (Extension.Equals(TEXT(".utxt"_SV), ESearchCase::IgnoreCase))
+	if (Extension.Equals(TEXTVIEW(".utxt"), ESearchCase::IgnoreCase))
 	{
 		return EPackageExtension::TextAsset;
 	}
-	if (Extension.Equals(TEXT(".utxtmap"_SV), ESearchCase::IgnoreCase))
+	if (Extension.Equals(TEXTVIEW(".utxtmap"), ESearchCase::IgnoreCase))
 	{
 		return EPackageExtension::TextMap;
 	}
-	if (Extension.Equals(TEXT(".uexp"_SV), ESearchCase::IgnoreCase))
+	if (Extension.Equals(TEXTVIEW(".uexp"), ESearchCase::IgnoreCase))
 	{
 		return EPackageExtension::Exports;
 	}
-	if (Extension.Equals(TEXT(".ubulk"_SV), ESearchCase::IgnoreCase))
+	if (Extension.Equals(TEXTVIEW(".ubulk"), ESearchCase::IgnoreCase))
 	{
 		// .m.ubulk
 		// .ubulk
-		FStringView MemoryMappedExtension(TEXT(".m"_SV));
+		FStringView MemoryMappedExtension(TEXTVIEW(".m"));
 		if (FPathViews::GetExtension(Filename.LeftChop(Extension.Len()), true /* bIncludeDot */).Equals(MemoryMappedExtension, ESearchCase::IgnoreCase))
 		{
 			if (OutExtensionStart)
@@ -175,11 +175,11 @@ EPackageExtension FPackagePath::ParseExtension(FStringView Filename, int32* OutE
 			return EPackageExtension::BulkDataDefault;
 		}
 	}
-	if (Extension.Equals(TEXT(".uptnl"_SV), ESearchCase::IgnoreCase))
+	if (Extension.Equals(TEXTVIEW(".uptnl"), ESearchCase::IgnoreCase))
 	{
 		return EPackageExtension::BulkDataOptional;
 	}
-	if (Extension.Equals(TEXT(".upayload"_SV), ESearchCase::IgnoreCase))
+	if (Extension.Equals(TEXTVIEW(".upayload"), ESearchCase::IgnoreCase))
 	{
 		return EPackageExtension::PayloadSidecar;
 	}
diff --git a/Engine/Source/Runtime/CoreUObject/Private/Misc/PackagePathTest.cpp b/Engine/Source/Runtime/CoreUObject/Private/Misc/PackagePathTest.cpp
--- a/Engine/Source/Runtime/CoreUObject/Private/Misc/PackagePathTest.cpp
+++ b/Engine/Source/Runtime/CoreUObject/Private/Misc/PackagePathTest.cpp
@@ -41,7 +41,7 @@ public:
 		auto TestExtension = [this](EPackageExtension Extension)
 		{
 			TStringBuilder<64> Path;
-			const FStringView PathWithoutExtension = TEXT("d:\\Root\\Dir\\Filename"_SV);
+			const FStringView PathWithoutExtension = TEXTVIEW("d:\\Root\\Dir\\Filename");
 			Path << PathWithoutExtension << LexToString(Extension);
 			int32 ExtensionStart;
 			EPackageExtension ParsedExtension = FPackagePath::ParseExtension(Path, &ExtensionStart);
@@ -65,7 +65,7 @@ public:
 		// Test Custom
 		{
 			TStringBuilder<64> Path;
-			const FStringView PathWithoutExtension = TEXT("d:\\Root\\Dir\\Filename"_SV);
+			const FStringView PathWithoutExtension = TEXTVIEW("d:\\Root\\Dir\\Filename");
 			Path << PathWithoutExtension << TEXT(".SomeOtherExtension");
 			int32 ExtensionStart;
 			EPackageExtension ParsedExtension = FPackagePath::ParseExtension(Path, &ExtensionStart);
diff --git a/Engine/Source/Runtime/CoreUObject/Private/UObject/PropertyBaseObject.cpp b/Engine/Source/Runtime/CoreUObject/Private/UObject/PropertyBaseObject.cpp
--- a/Engine/Source/Runtime/CoreUObject/Private/UObject/PropertyBaseObject.cpp
+++ b/Engine/Source/Runtime/CoreUObject/Private/UObject/PropertyBaseObject.cpp
@@ -313,7 +313,7 @@ bool FObjectPropertyBase::ParseObjectPropertyValue(const FProperty* Property, UO
 		return false;
 	}
 
-	if ( Temp == TEXT("None"_SV) )
+	if ( Temp == TEXTVIEW("None") )
 	{
 		out_ResolvedValue = nullptr;
 	}
diff --git a/Engine/Source/Runtime/CoreUObject/Private/UObject/SoftObjectPath.cpp b/Engine/Source/Runtime/CoreUObject/Private/UObject/SoftObjectPath.cpp
--- a/Engine/Source/Runtime/CoreUObject/Private/UObject/SoftObjectPath.cpp
+++ b/Engine/Source/Runtime/CoreUObject/Private/UObject/SoftObjectPath.cpp
@@ -333,7 +333,7 @@ bool FSoftObjectPath::ImportTextItem(const TCHAR*& Buffer, int32 PortFlags, UObj
 		return false;
 	}
 	Buffer = NewBuffer;
-	if (ImportedPath == TEXT("None"_SV))
+	if (ImportedPath == TEXTVIEW("None"))
 	{
 		Reset();
 	}
diff --git a/Engine/Source/Runtime/Launch/Private/LaunchEngineLoop.cpp b/Engine/Source/Runtime/Launch/Private/LaunchEngineLoop.cpp
--- a/Engine/Source/Runtime/Launch/Private/LaunchEngineLoop.cpp
+++ b/Engine/Source/Runtime/Launch/Private/LaunchEngineLoop.cpp
@@ -685,7 +685,7 @@ void LaunchFixProjectPathCase()
 			const FStringView CurrentPathComponent = FPathViews::GetCleanFilename(PendingFix);
 
 			// Skip over all segments that are either empty or contain relative transforms or start with the volume separator, they should remain as-is
-			const bool bIsIgnoredSegment = CurrentPathComponent.IsEmpty() || CurrentPathComponent.Equals(TEXT("."_SV)) || CurrentPathComponent.Equals(TEXT(".."_SV)) || CurrentPathComponent.EndsWith(TEXT(':'));
+			const bool bIsIgnoredSegment = CurrentPathComponent.IsEmpty() || CurrentPathComponent.Equals(TEXTVIEW(".")) || CurrentPathComponent.Equals(TEXTVIEW("..")) || CurrentPathComponent.EndsWith(TEXT(':'));
 			if (!bIsIgnoredSegment)
 			{
 				// Temporarily null-terminate the current path component for the system call
