Fixed IncludeTool warning.

#rnx

diff --git a/Engine/Source/Developer/Zen/Public/ZenServerInterface.h b/Engine/Source/Developer/Zen/Public/ZenServerInterface.h
--- a/Engine/Source/Developer/Zen/Public/ZenServerInterface.h
+++ b/Engine/Source/Developer/Zen/Public/ZenServerInterface.h
@@ -14,6 +14,10 @@
 #include "ZenGlobals.h"
 #include "ZenServerHttp.h"
 
+#if UE_WITH_ZEN
+#	include "ZenStatistics.h"
+#endif
+
 #define UE_API ZEN_API
 
 namespace UE::Zen
@@ -51,12 +55,10 @@ private:
 	bool TryApplyAutoLaunchOverride();
 };
 
-};
+}
 
 #if UE_WITH_ZEN
 
-#include "ZenStatistics.h"
-
 namespace UE::Zen
 {
 
