Blocking Zen HTTP stats request now runs async using futures

diff --git a/Engine/Source/Developer/Zen/Internal/ZenServerInterface.cpp b/Engine/Source/Developer/Zen/Internal/ZenServerInterface.cpp
--- a/Engine/Source/Developer/Zen/Internal/ZenServerInterface.cpp
+++ b/Engine/Source/Developer/Zen/Internal/ZenServerInterface.cpp
@@ -4,6 +4,7 @@
 
 #include "ZenBackendUtils.h"
 #include "ZenSerialization.h"
+#include "ZenServerHttp.h"
 
 #include "Dom/JsonValue.h"
 #include "HAL/FileManager.h"
@@ -11,6 +12,7 @@
 #include "HAL/PlatformMisc.h"
 #include "HAL/PlatformProcess.h"
 #include "HAL/PlatformTime.h"
+#include "Async/Async.h"
 #include "Logging/LogScopedCategoryAndVerbosityOverride.h"
 #include "Misc/App.h"
 #include "Misc/ConfigCacheIni.h"
@@ -1044,116 +1046,134 @@ FZenServiceInstance::AutoLaunch(const FServiceAutoLaunchSettings& InSettings, FS
 }
 
 bool 
-FZenServiceInstance::GetStats( FZenStats& stats ) const
+FZenServiceInstance::GetStats(FZenStats& Stats)
 {
 	check(IsInGameThread());
 
-	const uint64 CurrentCycles = FPlatformTime::Cycles64();
-	if (StatsRequest.IsSet() && FPlatformTime::ToSeconds(CurrentCycles - LastStatsTime) < 0.5)
+	// If we've already requested a stats and they are ready then grab them
+	if ( StatsRequest.IsReady() == true )
 	{
-		stats = LastStats;
-		return true;
-	}
+		LastStats		= StatsRequest.Get();
+		LastStatsTime	= FPlatformTime::Cycles64();
 
-	TStringBuilder<128> ZenDomain;
-	ZenDomain << HostName << TEXT(":") << Port;
-	if (!StatsRequest.IsSet())
-	{
-		StatsRequest.Emplace(ZenDomain.ToString(), false);
+		StatsRequest.Reset();
 	}
-	UE::Zen::FZenHttpRequest& Request = StatsRequest.GetValue();
-	Request.Reset();
+	
+	// Make a copy of the last updated stats
+	Stats = LastStats;
 
-	TArray64<uint8> GetBuffer;
-	FZenHttpRequest::Result Result = Request.PerformBlockingDownload(TEXTVIEW("/stats/z$"), &GetBuffer, Zen::EContentType::CbObject);
+	const uint64 CurrentTime = FPlatformTime::Cycles64();
+	const float MinTimeBetweenRequestsInSeconds = 0.5f;
+	const float DeltaTimeInSeconds = FPlatformTime::ToSeconds(CurrentTime - LastStatsTime);
 
-	if (Result == Zen::FZenHttpRequest::Result::Success && Request.GetResponseCode() == 200)
+	if ( StatsRequest.IsValid()==false && DeltaTimeInSeconds > MinTimeBetweenRequestsInSeconds)
 	{
-		FCbObjectView RootObjectView(GetBuffer.GetData());
+#if WITH_EDITOR
+		EAsyncExecution ThreadPool = EAsyncExecution::LargeThreadPool;
+#else
+		EAsyncExecution ThreadPool = EAsyncExecution::ThreadPool;
+#endif
+		// We've not got any requests in flight and we've met a given time requirement for requests
+		StatsRequest = Async(ThreadPool, [this]
+			{
+				TStringBuilder<128> ZenDomain;
+				ZenDomain << HostName << TEXT(":") << Port;
+				UE::Zen::FZenHttpRequest Request(ZenDomain.ToString(), false);
 
-		FCbObjectView RequestsObjectView = RootObjectView["requests"].AsObjectView();
-		FZenRequestStats& RequestStats = stats.RequestStats;
-		
-		RequestStats.Count = RequestsObjectView["count"].AsInt64();
-		RequestStats.RateMean = RequestsObjectView["rate_mean"].AsDouble();
-		RequestStats.TAverage = RequestsObjectView["t_avg"].AsDouble();
-		RequestStats.TMin = RequestsObjectView["t_min"].AsDouble();
-		RequestStats.TMax = RequestsObjectView["t_max"].AsDouble();
-
-		FCbObjectView CacheObjectView = RootObjectView["cache"].AsObjectView();
-		FZenCacheStats& CacheStats = stats.CacheStats;
-
-		CacheStats.Hits = CacheObjectView["hits"].AsInt64();
-		CacheStats.Misses = CacheObjectView["misses"].AsInt64();
-		CacheStats.HitRatio = CacheObjectView["hit_ratio"].AsDouble();
-		CacheStats.UpstreamHits = CacheObjectView["upstream_hits"].AsInt64();
-		CacheStats.UpstreamRatio = CacheObjectView["upstream_ratio"].AsDouble();
-
-		FCbObjectView CacheSizeObjectView = CacheObjectView["size"].AsObjectView();
-		FZenCacheSizeStats& CacheSizeStats = CacheStats.Size;
-		CacheSizeStats.Disk = CacheSizeObjectView["disk"].AsDouble();
-		CacheSizeStats.Memory = CacheSizeObjectView["memory"].AsDouble();
-
-		FCbObjectView UpstreamObjectView = RootObjectView["upstream"].AsObjectView();
-		FZenUpstreamStats& UpstreamStats = stats.UpstreamStats;
-
-		UpstreamStats.Reading = UpstreamObjectView["reading"].AsBool();
-		UpstreamStats.Writing = UpstreamObjectView["writing"].AsBool();
-		UpstreamStats.WorkerThreads = UpstreamObjectView["worker_threads"].AsInt64();
-		UpstreamStats.QueueCount = UpstreamObjectView["queue_count"].AsInt64();
-		UpstreamStats.TotalUploadedMB = 0.0;
-		UpstreamStats.TotalDownloadedMB = 0.0;
-
-		FCbObjectView UpstreamRequestObjectView = RootObjectView["upstream_gets"].AsObjectView();
-		FZenRequestStats& UpstreamRequestStats = stats.UpstreamRequestStats;
-
-		UpstreamRequestStats.Count = UpstreamRequestObjectView["count"].AsInt64();
-		UpstreamRequestStats.RateMean = UpstreamRequestObjectView["rate_mean"].AsDouble();
-		UpstreamRequestStats.TAverage = UpstreamRequestObjectView["t_avg"].AsDouble();
-		UpstreamRequestStats.TMin = UpstreamRequestObjectView["t_min"].AsDouble();
-		UpstreamRequestStats.TMax = UpstreamRequestObjectView["t_max"].AsDouble();
-
-		FCbArrayView EndpPointArrayView = UpstreamObjectView["endpoints"].AsArrayView();
-
-		for (FCbFieldView FieldView : EndpPointArrayView)
-		{
-			FCbObjectView EndPointView = FieldView.AsObjectView();
-			FZenEndPointStats EndPointStats;
+				TArray64<uint8> GetBuffer;
+				FZenHttpRequest::Result Result = Request.PerformBlockingDownload(TEXT("/stats/z$"_SV), &GetBuffer, Zen::EContentType::CbObject);
 
-			EndPointStats.Name = FString(EndPointView["name"].AsString());
-			EndPointStats.Url = FString(EndPointView["url"].AsString());
-			EndPointStats.Health = FString(EndPointView["state"].AsString());
+				FZenStats Stats;
 
-			if (FCbObjectView Cache = EndPointView["cache"].AsObjectView())
-			{
-				EndPointStats.HitRatio = Cache["hit_ratio"].AsDouble();
-				EndPointStats.UploadedMB = Cache["put_bytes"].AsDouble() / 1024.0 / 1024.0;
-				EndPointStats.DownloadedMB = Cache["get_bytes"].AsDouble() / 1024.0 / 1024.0;
-				EndPointStats.ErrorCount = Cache["error_count"].AsInt64();
-			}
-			
-			UpstreamStats.TotalUploadedMB += EndPointStats.UploadedMB;
-			UpstreamStats.TotalDownloadedMB += EndPointStats.DownloadedMB;
+				if (Result == Zen::FZenHttpRequest::Result::Success && Request.GetResponseCode() == 200)
+				{
+					FCbObjectView RootObjectView(GetBuffer.GetData());
 
-			UpstreamStats.EndPointStats.Push(EndPointStats);
-		}
+					FCbObjectView RequestsObjectView = RootObjectView["requests"].AsObjectView();
+					FZenRequestStats& RequestStats = Stats.RequestStats;
 
-		FCbObjectView CASObjectView = RootObjectView["cas"].AsObjectView();	
-		FCbObjectView CASSizeObjectView = CASObjectView["size"].AsObjectView();
+					RequestStats.Count = RequestsObjectView["count"].AsInt64();
+					RequestStats.RateMean = RequestsObjectView["rate_mean"].AsDouble();
+					RequestStats.TAverage = RequestsObjectView["t_avg"].AsDouble();
+					RequestStats.TMin = RequestsObjectView["t_min"].AsDouble();
+					RequestStats.TMax = RequestsObjectView["t_max"].AsDouble();
 
-		FZenCASSizeStats& CASSizeStats = stats.CASStats.Size;
+					FCbObjectView CacheObjectView = RootObjectView["cache"].AsObjectView();
+					FZenCacheStats& CacheStats = Stats.CacheStats;
 
-		CASSizeStats.Tiny = CASSizeObjectView["tiny"].AsInt64();
-		CASSizeStats.Small = CASSizeObjectView["small"].AsInt64();
-		CASSizeStats.Large = CASSizeObjectView["large"].AsInt64();
-		CASSizeStats.Total = CASSizeObjectView["total"].AsInt64();
+					CacheStats.Hits = CacheObjectView["hits"].AsInt64();
+					CacheStats.Misses = CacheObjectView["misses"].AsInt64();
+					CacheStats.HitRatio = CacheObjectView["hit_ratio"].AsDouble();
+					CacheStats.UpstreamHits = CacheObjectView["upstream_hits"].AsInt64();
+					CacheStats.UpstreamRatio = CacheObjectView["upstream_ratio"].AsDouble();
 
-		LastStats = stats;
-		LastStatsTime = CurrentCycles;
-		return true;
+					FCbObjectView CacheSizeObjectView = CacheObjectView["size"].AsObjectView();
+					FZenCacheSizeStats& CacheSizeStats = CacheStats.Size;
+					CacheSizeStats.Disk = CacheSizeObjectView["disk"].AsDouble();
+					CacheSizeStats.Memory = CacheSizeObjectView["memory"].AsDouble();
+
+					FCbObjectView UpstreamObjectView = RootObjectView["upstream"].AsObjectView();
+					FZenUpstreamStats& UpstreamStats = Stats.UpstreamStats;
+
+					UpstreamStats.Reading = UpstreamObjectView["reading"].AsBool();
+					UpstreamStats.Writing = UpstreamObjectView["writing"].AsBool();
+					UpstreamStats.WorkerThreads = UpstreamObjectView["worker_threads"].AsInt64();
+					UpstreamStats.QueueCount = UpstreamObjectView["queue_count"].AsInt64();
+					UpstreamStats.TotalUploadedMB = 0.0;
+					UpstreamStats.TotalDownloadedMB = 0.0;
+
+					FCbObjectView UpstreamRequestObjectView = RootObjectView["upstream_gets"].AsObjectView();
+					FZenRequestStats& UpstreamRequestStats = Stats.UpstreamRequestStats;
+
+					UpstreamRequestStats.Count = UpstreamRequestObjectView["count"].AsInt64();
+					UpstreamRequestStats.RateMean = UpstreamRequestObjectView["rate_mean"].AsDouble();
+					UpstreamRequestStats.TAverage = UpstreamRequestObjectView["t_avg"].AsDouble();
+					UpstreamRequestStats.TMin = UpstreamRequestObjectView["t_min"].AsDouble();
+					UpstreamRequestStats.TMax = UpstreamRequestObjectView["t_max"].AsDouble();
+
+					FCbArrayView EndpPointArrayView = UpstreamObjectView["endpoints"].AsArrayView();
+
+					for (FCbFieldView FieldView : EndpPointArrayView)
+					{
+						FCbObjectView EndPointView = FieldView.AsObjectView();
+						FZenEndPointStats EndPointStats;
+
+						EndPointStats.Name = FString(EndPointView["name"].AsString());
+						EndPointStats.Url = FString(EndPointView["url"].AsString());
+						EndPointStats.Health = FString(EndPointView["state"].AsString());
+
+						if (FCbObjectView Cache = EndPointView["cache"].AsObjectView())
+						{
+							EndPointStats.HitRatio = Cache["hit_ratio"].AsDouble();
+							EndPointStats.UploadedMB = Cache["put_bytes"].AsDouble() / 1024.0 / 1024.0;
+							EndPointStats.DownloadedMB = Cache["get_bytes"].AsDouble() / 1024.0 / 1024.0;
+							EndPointStats.ErrorCount = Cache["error_count"].AsInt64();
+						}
+
+						UpstreamStats.TotalUploadedMB += EndPointStats.UploadedMB;
+						UpstreamStats.TotalDownloadedMB += EndPointStats.DownloadedMB;
+
+						UpstreamStats.EndPointStats.Push(EndPointStats);
+					}
+
+					FCbObjectView CASObjectView = RootObjectView["cas"].AsObjectView();
+					FCbObjectView CASSizeObjectView = CASObjectView["size"].AsObjectView();
+
+					FZenCASSizeStats& CASSizeStats = Stats.CASStats.Size;
+
+					CASSizeStats.Tiny = CASSizeObjectView["tiny"].AsInt64();
+					CASSizeStats.Small = CASSizeObjectView["small"].AsInt64();
+					CASSizeStats.Large = CASSizeObjectView["large"].AsInt64();
+					CASSizeStats.Total = CASSizeObjectView["total"].AsInt64();
+
+					Stats.IsValid = true;
+				}
+
+				return Stats;
+			});
 	}
 
-	return false;
+	return Stats.IsValid;
 }
 
 #endif // UE_WITH_ZEN
diff --git a/Engine/Source/Developer/Zen/Public/ZenServerInterface.h b/Engine/Source/Developer/Zen/Public/ZenServerInterface.h
--- a/Engine/Source/Developer/Zen/Public/ZenServerInterface.h
+++ b/Engine/Source/Developer/Zen/Public/ZenServerInterface.h
@@ -12,7 +12,7 @@
 #include "Serialization/JsonWriter.h"
 #include "Templates/UniquePtr.h"
 #include "ZenGlobals.h"
-#include "ZenServerHttp.h"
+#include "Async/Future.h"
 
 #if UE_WITH_ZEN
 #	include "ZenStatistics.h"
@@ -120,7 +120,7 @@ public:
 	inline const TCHAR* GetHostName() const { return *HostName; }
 	inline uint16 GetPort() const { return Port; }
 	inline const FServiceSettings& GetServiceSettings() const { return Settings; }
-	UE_API bool GetStats(FZenStats& stats) const;
+	UE_API bool GetStats(FZenStats& Stats);
 	UE_API bool IsServiceRunning();
 	UE_API bool IsServiceReady();
 	UE_API bool IsServiceRunningLocally() const { return bIsRunningLocally; }
@@ -134,7 +134,7 @@ private:
 	FString ConditionalUpdateLocalInstall();
 	static bool AutoLaunch(const FServiceAutoLaunchSettings& InSettings, FString&& ExecutablePath, FString& OutHostName, uint16& OutPort);
 
-	mutable TOptional<FZenHttpRequest> StatsRequest;
+	mutable TFuture<FZenStats> StatsRequest;
 	mutable FZenStats LastStats;
 	mutable uint64 LastStatsTime = 0;
 
diff --git a/Engine/Source/Developer/Zen/Public/ZenStatistics.h b/Engine/Source/Developer/Zen/Public/ZenStatistics.h
--- a/Engine/Source/Developer/Zen/Public/ZenStatistics.h
+++ b/Engine/Source/Developer/Zen/Public/ZenStatistics.h
@@ -73,7 +73,7 @@ namespace UE::Zen {
 		FZenUpstreamStats UpstreamStats;
 		FZenRequestStats UpstreamRequestStats;
 		FZenCASStats CASStats;
-
+		bool IsValid = false;
 	};
 
 } // namespace UE::Zen
