Add command-line option for passing config file path when launching ZenServer.

#rnx

diff --git a/Engine/Source/Developer/Zen/Internal/ZenServerInterface.cpp b/Engine/Source/Developer/Zen/Internal/ZenServerInterface.cpp
--- a/Engine/Source/Developer/Zen/Internal/ZenServerInterface.cpp
+++ b/Engine/Source/Developer/Zen/Internal/ZenServerInterface.cpp
@@ -865,6 +865,16 @@ FZenServiceInstance::AutoLaunch(const FServiceAutoLaunchSettings& InSettings, FS
 			}
 		}
 
+		FString CfgCommandLineOverrideValue;
+		if (FParse::Value(FCommandLine::Get(), TEXT("ZenCfgPath="), CfgCommandLineOverrideValue))
+		{
+			if (!CfgCommandLineOverrideValue.IsEmpty())
+			{
+				Parms.Appendf(TEXT(" --config \"%s\""),
+					*FPaths::ConvertRelativePathToFull(CfgCommandLineOverrideValue));
+			}
+		}
+
 		if (InSettings.bLimitProcessLifetime)
 		{
 			Parms.Appendf(TEXT(" --owner-pid %d"),
