Increase the time window that ZenServer has to launch within before it becomes an error for unattended executions.

#rnx

diff --git a/Engine/Source/Developer/Zen/Internal/ZenServerInterface.cpp b/Engine/Source/Developer/Zen/Internal/ZenServerInterface.cpp
--- a/Engine/Source/Developer/Zen/Internal/ZenServerInterface.cpp
+++ b/Engine/Source/Developer/Zen/Internal/ZenServerInterface.cpp
@@ -990,9 +990,9 @@ FZenServiceInstance::AutoLaunch(const FServiceAutoLaunchSettings& InSettings, FS
 			}
 
 			double ZenWaitDuration = FPlatformTime::ToSeconds64(FPlatformTime::Cycles64() - ZenWaitStartTime);
-			if (ZenWaitDuration < 3.0)
+			if (ZenWaitDuration < 5.0)
 			{
-				// Initial 3 second window of higher frequency checks
+				// Initial 5 second window of higher frequency checks
 				FPlatformProcess::Sleep(0.01f);
 			}
 			else
