Ensure that when subprocesses are launched (turnkey or launcher) that the ZenDataPath is inherited from the parent process.  This avoids errors experienced in the CookOnTheFly flow when the parent process and subprocess differ in their expectation for Zen arguments.

#rnx

diff --git a/Engine/Source/Developer/Zen/Internal/ZenServerInterface.cpp b/Engine/Source/Developer/Zen/Internal/ZenServerInterface.cpp
--- a/Engine/Source/Developer/Zen/Internal/ZenServerInterface.cpp
+++ b/Engine/Source/Developer/Zen/Internal/ZenServerInterface.cpp
@@ -131,6 +131,16 @@ DetermineDataPath(const TCHAR* ConfigSection, FString& DataPath)
 		return;
 	}
 
+	// Zen subprocess environment
+	FString SubprocessDataPathEnvOverrideValue;
+	SubprocessDataPathEnvOverrideValue = FPlatformMisc::GetEnvironmentVariable(TEXT("UE-ZenSubprocessDataPath"));
+	if (!SubprocessDataPathEnvOverrideValue.IsEmpty())
+	{
+		DataPath = NormalizeDataPath(SubprocessDataPathEnvOverrideValue);
+		UE_LOG(LogZenServiceInstance, Log, TEXT("Found subprocess environment variable UE-ZenSubprocessDataPath=%s"), *SubprocessDataPathEnvOverrideValue);
+		return;
+	}
+
 	// Zen registry/stored
 	FString DataPathEnvOverrideValue;
 	if (FPlatformMisc::GetStoredValue(TEXT("Epic Games"), TEXT("Zen"), TEXT("DataPath"), DataPathEnvOverrideValue))
@@ -244,6 +254,9 @@ FServiceSettings::ReadFromConfig()
 			ReadUInt16FromConfig(AutoLaunchConfigSection, TEXT("DesiredPort"), AutoLaunchSettings.DesiredPort, GEngineIni);
 			GConfig->GetBool(AutoLaunchConfigSection, TEXT("ShowConsole"), AutoLaunchSettings.bShowConsole, GEngineIni);
 			GConfig->GetBool(AutoLaunchConfigSection, TEXT("LimitProcessLifetime"), AutoLaunchSettings.bLimitProcessLifetime, GEngineIni);
+
+			// Ensure that the zen data path is inherited by subprocesses
+			FPlatformMisc::SetEnvironmentVar(TEXT("UE-ZenSubprocessDataPath"), *AutoLaunchSettings.DataPath);
 		}
 	}
 	else
