Fix bugs with the installation of the symbol file and crashpad_handler when launching ZenServer.

#rnx

diff --git a/Engine/Source/Developer/Zen/Internal/ZenServerInterface.cpp b/Engine/Source/Developer/Zen/Internal/ZenServerInterface.cpp
--- a/Engine/Source/Developer/Zen/Internal/ZenServerInterface.cpp
+++ b/Engine/Source/Developer/Zen/Internal/ZenServerInterface.cpp
@@ -793,6 +793,8 @@ FZenServiceInstance::ConditionalUpdateLocalInstall()
 	FDateTime InTreeFileTime;
 	FDateTime InstallFileTime;
 	FileManager.GetTimeStampPair(*InTreeFilePath, *InstallFilePath, InTreeFileTime, InstallFileTime);
+
+	bool bMainExecutableUpdated = false;
 	if (InTreeFileTime > InstallFileTime)
 	{
 		if (IsProcessActive(*InstallFilePath))
@@ -819,15 +821,26 @@ FZenServiceInstance::ConditionalUpdateLocalInstall()
 
 		// Even after waiting for the lock file to be removed, the executable may have a period where it can't be overwritten as the process shuts down
 		// so any attempt to overwrite it should have some tolerance for retrying.
-		bool bExecutableCopySucceeded = AttemptFileCopyWithRetries(*InstallFilePath, *InTreeFilePath, 5.0);
-		checkf(bExecutableCopySucceeded, TEXT("Failed to copy zenserver to install location '%s'."), *InstallFilePath);
+		bMainExecutableUpdated = AttemptFileCopyWithRetries(*InstallFilePath, *InTreeFilePath, 5.0);
+		checkf(bMainExecutableUpdated, TEXT("Failed to copy zenserver to install location '%s'."), *InstallFilePath);
+	}
 
 #if PLATFORM_WINDOWS
-		FString InTreeSymbolFilePath = FPaths::ChangeExtension(InTreeFilePath, TEXT("pdb"));
-		FString InstallSymbolFilePath = FPaths::ChangeExtension(InstallFilePath, TEXT("pdb"));
-		bool bSymbolCopySucceeded = AttemptFileCopyWithRetries(*InstallFilePath, *InTreeFilePath, 1.0);
-		checkf(bSymbolCopySucceeded, TEXT("Failed to copy zenserver symbols to install location '%s'."), *InstallSymbolFilePath);
+	FString InTreeSymbolFilePath = FPaths::ChangeExtension(InTreeFilePath, TEXT("pdb"));
+	FString InstallSymbolFilePath = FPaths::ChangeExtension(InstallFilePath, TEXT("pdb"));
+
+	if (FileManager.FileExists(*InTreeSymbolFilePath) && (bMainExecutableUpdated || !FileManager.FileExists(*InstallSymbolFilePath)))
+	{
+		AttemptFileCopyWithRetries(*InstallSymbolFilePath, *InTreeSymbolFilePath, 1.0);
+	}
 #endif
+
+	FString InTreeCrashpadHandlerFilePath = FPaths::ConvertRelativePathToFull(FPlatformProcess::GenerateApplicationPath(TEXT("crashpad_handler"), EBuildConfiguration::Development));
+	FString InstallCrashpadHandlerFilePath = FPaths::ConvertRelativePathToFull(FPaths::Combine(FPlatformProcess::ApplicationSettingsDir(), TEXT("Zen\\Install"), FString(FPathViews::GetCleanFilename(InTreeCrashpadHandlerFilePath))));
+
+	if (FileManager.FileExists(*InTreeCrashpadHandlerFilePath) && (bMainExecutableUpdated || !FileManager.FileExists(*InstallCrashpadHandlerFilePath)))
+	{
+		AttemptFileCopyWithRetries(*InstallCrashpadHandlerFilePath, *InTreeCrashpadHandlerFilePath, 1.0);
 	}
 
 	FPaths::MakePlatformFilename(InstallFilePath);
