Set the working directory when launching zenserver to avoid the working directory being inherited from the running process and leading to issues with that path being in use when it needs to be deleted.

diff --git a/Engine/Source/Developer/Zen/Internal/ZenServerInterface.cpp b/Engine/Source/Developer/Zen/Internal/ZenServerInterface.cpp
--- a/Engine/Source/Developer/Zen/Internal/ZenServerInterface.cpp
+++ b/Engine/Source/Developer/Zen/Internal/ZenServerInterface.cpp
@@ -853,7 +853,9 @@ FZenServiceInstance::AutoLaunch(const FServiceAutoLaunchSettings& InSettings, FS
 	int16 DesiredPort = InSettings.DesiredPort;
 	IFileManager& FileManager = IFileManager::Get();
 	const FString LockFilePath = FPaths::Combine(InSettings.DataPath, TEXT(".lock"));
-	const FString CmdLineFilePath = FPaths::Combine(InSettings.DataPath, TEXT(".cmdline"));
+	const FString ExecutionContextFilePath = FPaths::Combine(InSettings.DataPath, TEXT(".runcontext"));
+
+	FString WorkingDirectory = FPaths::GetPath(ExecutablePath);
 
 	bool bReUsingExistingInstance = false;
 
@@ -872,9 +874,9 @@ FZenServiceInstance::AutoLaunch(const FServiceAutoLaunchSettings& InSettings, FS
 		}
 
 		bool bCurrentInstanceUsable = false;
-		FString DesiredCmdLine = FString::Printf(TEXT("%s %s"), *ExecutablePath, *DetermineCmdLineWithoutTransientComponents(InSettings, CurrentPort));
-		FString CurrentCmdLine;
-		if (FFileHelper::LoadFileToString(CurrentCmdLine, *CmdLineFilePath) && (DesiredCmdLine == CurrentCmdLine))
+		FString DesiredExecutionContext = FString::Printf(TEXT("%s %s") LINE_TERMINATOR TEXT("%s"), *ExecutablePath, *DetermineCmdLineWithoutTransientComponents(InSettings, CurrentPort), *WorkingDirectory);
+		FString CurrentExecutionContext;
+		if (FFileHelper::LoadFileToString(CurrentExecutionContext, *ExecutionContextFilePath) && (DesiredExecutionContext == CurrentExecutionContext))
 		{
 			DesiredPort = CurrentPort;
 			bReUsingExistingInstance = true;
@@ -929,7 +931,7 @@ FZenServiceInstance::AutoLaunch(const FServiceAutoLaunchSettings& InSettings, FS
 
 			FString CommandLine = FString::Printf(TEXT("\"%s\" %s"), *ExecutablePath, *Parms);
 			PROCESS_INFORMATION ProcInfo;
-			if (CreateProcess(NULL, CommandLine.GetCharArray().GetData(), nullptr, nullptr, false, (::DWORD)(NORMAL_PRIORITY_CLASS | DETACHED_PROCESS), nullptr, nullptr, &StartupInfo, &ProcInfo))
+			if (CreateProcess(NULL, CommandLine.GetCharArray().GetData(), nullptr, nullptr, false, (::DWORD)(NORMAL_PRIORITY_CLASS | DETACHED_PROCESS), nullptr, WorkingDirectory.GetCharArray().GetData(), &StartupInfo, &ProcInfo))
 			{
 				::CloseHandle(ProcInfo.hThread);
 				Proc = FProcHandle(ProcInfo.hProcess);
@@ -960,7 +962,6 @@ FZenServiceInstance::AutoLaunch(const FServiceAutoLaunchSettings& InSettings, FS
 			bool bLaunchReallyHidden = !InSettings.bShowConsole;
 			uint32* OutProcessID = nullptr;
 			int32 PriorityModifier = 0;
-			const TCHAR* OptionalWorkingDirectory = nullptr;
 			void* PipeWriteChild = nullptr;
 			void* PipeReadChild = nullptr;
 			Proc = FPlatformProcess::CreateProc(
@@ -971,15 +972,15 @@ FZenServiceInstance::AutoLaunch(const FServiceAutoLaunchSettings& InSettings, FS
 				bLaunchReallyHidden,
 				OutProcessID,
 				PriorityModifier,
-				OptionalWorkingDirectory,
+				*WorkingDirectory,
 				PipeWriteChild,
 				PipeReadChild);
 		}
 #endif
 		if (!bProcessIsLive)
 		{
-			FString ExecutedCmdLine = FString::Printf(TEXT("%s %s"), *ExecutablePath, *ParmsWithoutTransients);
-			FFileHelper::SaveStringToFile(ExecutedCmdLine,*CmdLineFilePath);
+			FString UsedExecutionContext = FString::Printf(TEXT("%s %s") LINE_TERMINATOR TEXT("%s"), *ExecutablePath, *ParmsWithoutTransients, *WorkingDirectory);
+			FFileHelper::SaveStringToFile(UsedExecutionContext,*ExecutionContextFilePath);
 		}
 
 		bProcessIsLive = Proc.IsValid();
