Allow support for receiving local file references for cache items
Companion PR in zen: https://github.com/EpicGames/zen/pull/184

diff --git a/Engine/Source/Developer/DerivedDataCache/Private/ZenCacheStore.cpp b/Engine/Source/Developer/DerivedDataCache/Private/ZenCacheStore.cpp
--- a/Engine/Source/Developer/DerivedDataCache/Private/ZenCacheStore.cpp
+++ b/Engine/Source/Developer/DerivedDataCache/Private/ZenCacheStore.cpp
@@ -151,6 +151,7 @@ private:
 	THttpUniquePtr<IHttpConnectionPool> ConnectionPool;
 	FHttpRequestQueue RequestQueue;
 	bool bIsUsable = false;
+	bool bIsLocalConnection = false;
 	int32 BatchPutMaxBytes = 1024*1024;
 	int32 CacheRecordBatchSize = 8;
 	int32 CacheChunksBatchSize = 8;
@@ -334,6 +335,10 @@ public:
 				{
 					BatchRequest << ANSITEXTVIEW("Method") << ANSITEXTVIEW("GetCacheRecords");
 					BatchRequest.AddInteger(ANSITEXTVIEW("Accept"), Zen::Http::kCbPkgMagic);
+					if (CacheStore.bIsLocalConnection)
+					{
+						BatchRequest.AddInteger(ANSITEXTVIEW("AcceptFlags"), static_cast<uint32_t>(Zen::Http::RpcAcceptOptions::kAllowLocalReferences));
+					}
 
 					BatchRequest.BeginObject(ANSITEXTVIEW("Params"));
 					{
@@ -608,6 +613,10 @@ public:
 			{
 				BatchRequest << ANSITEXTVIEW("Method") << ANSITEXTVIEW("GetCacheValues");
 				BatchRequest.AddInteger(ANSITEXTVIEW("Accept"), Zen::Http::kCbPkgMagic);
+				if (CacheStore.bIsLocalConnection)
+				{
+					BatchRequest.AddInteger(ANSITEXTVIEW("AcceptFlags"), static_cast<uint32_t>(Zen::Http::RpcAcceptOptions::kAllowLocalReferences));
+				}
 
 				BatchRequest.BeginObject(ANSITEXTVIEW("Params"));
 				{
@@ -757,6 +766,10 @@ public:
 			{
 				BatchRequest << ANSITEXTVIEW("Method") << "GetCacheChunks";
 				BatchRequest.AddInteger(ANSITEXTVIEW("Accept"), Zen::Http::kCbPkgMagic);
+				if (CacheStore.bIsLocalConnection)
+				{
+					BatchRequest.AddInteger(ANSITEXTVIEW("AcceptFlags"), static_cast<uint32_t>(Zen::Http::RpcAcceptOptions::kAllowLocalReferences));
+				}
 
 				BatchRequest.BeginObject(ANSITEXTVIEW("Params"));
 				{
@@ -1018,6 +1031,7 @@ FZenCacheStore::FZenCacheStore(
 		RequestQueue = FHttpRequestQueue(*ConnectionPool, ClientParams);
 
 		bIsUsable = true;
+		bIsLocalConnection = ZenService.GetInstance().IsServiceRunningLocally();
 
 		// Issue a request for stats as it will be fetched asynchronously and issuing now makes them available sooner for future callers.
 		Zen::FZenStats ZenStats;
diff --git a/Engine/Source/Developer/Zen/Internal/ZenSerialization.cpp b/Engine/Source/Developer/Zen/Internal/ZenSerialization.cpp
--- a/Engine/Source/Developer/Zen/Internal/ZenSerialization.cpp
+++ b/Engine/Source/Developer/Zen/Internal/ZenSerialization.cpp
@@ -143,10 +143,21 @@ struct CbAttachmentEntry
 	{
 		IsCompressed = (1u << 0),	// Is marshaled using compressed buffer storage format
 		IsObject = (1u << 1),		// Is compact binary object
+		kIsError = (1u << 2),	    // Is error (compact binary formatted) object
+		kIsLocalRef = (1u << 3),	// Is "local reference"
 	};
 
 };
 
+struct CbAttachmentReferenceHeader
+{
+	uint64_t PayloadByteOffset = 0;
+	uint64_t PayloadByteSize = ~0u;
+	uint16_t AbsolutePathLength = 0;
+
+	// This header will be followed by UTF8 encoded absolute path to backing file
+};
+
 void SaveCbPackage(const FCbPackage& Package, FArchive& Ar)
 {
 	TConstArrayView<FCbAttachment> Attachments = Package.GetAttachments();
@@ -263,7 +274,42 @@ bool TryLoadCbPackage(FCbPackage& Package, FArchive& Ar, FCbBufferAllocator Allo
 		FUniqueBuffer AttachmentData = FUniqueBuffer::Alloc(Entry.AttachmentSize);
 		Ar.Serialize(AttachmentData.GetData(), AttachmentData.GetSize());
 
-		if (Entry.Flags & CbAttachmentEntry::IsCompressed)
+		if (Entry.Flags & CbAttachmentEntry::kIsLocalRef)
+		{
+			// Marshal local reference - a "pointer" to the chunk backing file
+
+			check(AttachmentData.GetSize() >= sizeof(CbAttachmentReferenceHeader));
+
+			if (!(Entry.Flags & CbAttachmentEntry::IsCompressed))
+			{
+				Ar.SetError();
+				return false;
+			}
+
+			const CbAttachmentReferenceHeader* AttachRefHdr = reinterpret_cast<CbAttachmentReferenceHeader*>(AttachmentData.GetData());
+			const UTF8CHAR* PathPointer = reinterpret_cast<const UTF8CHAR*>(AttachRefHdr + 1);
+
+			check(AttachmentData.GetSize() >= (sizeof(CbAttachmentReferenceHeader) + AttachRefHdr->AbsolutePathLength));
+
+			FString Path(FUtf8StringView(PathPointer, static_cast<uint32_t>(AttachRefHdr->AbsolutePathLength)));
+			TUniquePtr<FArchive> ChunkReader(IFileManager::Get().CreateFileReader(*Path, FILEREAD_Silent));
+			if (!ChunkReader)
+			{
+				// File does not exist on disk - treat it as missing
+				++Index;
+				continue;
+			}
+			FCompressedBuffer CompBuf = FCompressedBuffer::Load(*ChunkReader);
+			if (!CompBuf)
+			{
+				// File has wrong format - treat it as an error
+				Ar.SetError();
+				return false;
+			}
+			FCbAttachment Attachment(MoveTemp(CompBuf));
+			Package.AddAttachment(Attachment);
+		}
+		else if (Entry.Flags & CbAttachmentEntry::IsCompressed)
 		{
 			FCompressedBuffer CompBuf(FCompressedBuffer::FromCompressed(AttachmentData.MoveToShared()));
 
diff --git a/Engine/Source/Developer/Zen/Internal/ZenSerialization.h b/Engine/Source/Developer/Zen/Internal/ZenSerialization.h
--- a/Engine/Source/Developer/Zen/Internal/ZenSerialization.h
+++ b/Engine/Source/Developer/Zen/Internal/ZenSerialization.h
@@ -28,6 +28,12 @@ namespace Http
 {
 static const uint32 kCbPkgMagic = 0xaa77aacc;
 	
+enum class RpcAcceptOptions : uint32_t
+{
+	kNone = 0,
+	kAllowLocalReferences = (1u << 0),
+};
+
 ZEN_API void SaveCbPackage(const FCbPackage& Package, FArchive& Ar);
 ZEN_API bool TryLoadCbPackage(FCbPackage& Package, FArchive& Ar, FCbBufferAllocator Allocator = FUniqueBuffer::Alloc);
 }
