Change handling of malformed CbPackages from hard error to treating it as a format error with logging

diff --git a/Engine/Source/Developer/Zen/Internal/ZenSerialization.cpp b/Engine/Source/Developer/Zen/Internal/ZenSerialization.cpp
--- a/Engine/Source/Developer/Zen/Internal/ZenSerialization.cpp
+++ b/Engine/Source/Developer/Zen/Internal/ZenSerialization.cpp
@@ -11,6 +11,8 @@
 
 #if UE_WITH_ZEN
 
+DEFINE_LOG_CATEGORY_STATIC(LogZenSerialization, Log, All);
+
 namespace UE::Zen
 {
 
@@ -69,6 +71,7 @@ bool TryLoadCbPackage(FCbPackage& Package, FArchive& Ar, FCbBufferAllocator Allo
 		FCbField ValueField = LoadCompactBinary(Ar, StackAllocator);
 		if (!ValueField)
 		{
+			UE_LOG(LogZenSerialization, Warning, TEXT("package is malformed, can't read compact binary data"));
 			Ar.SetError();
 			return false;
 		}
@@ -86,6 +89,7 @@ bool TryLoadCbPackage(FCbPackage& Package, FArchive& Ar, FCbBufferAllocator Allo
 				const FIoHash& Hash = HashField.AsAttachment();
 				if (HashField.HasError() || FIoHash::HashBuffer(Buffer) != Hash)
 				{
+					UE_LOG(LogZenSerialization, Warning, TEXT("package attachment has malformed/invalid attachment hash field"));
 					Ar.SetError();
 					return false;
 				}
@@ -104,6 +108,7 @@ bool TryLoadCbPackage(FCbPackage& Package, FArchive& Ar, FCbBufferAllocator Allo
 			FCbObject Object = ValueField.AsObject();
 			if (ValueField.HasError())
 			{
+				UE_LOG(LogZenSerialization, Warning, TEXT("package attachment has malformed/invalid object field"));
 				Ar.SetError();
 				return false;
 			}
@@ -114,6 +119,7 @@ bool TryLoadCbPackage(FCbPackage& Package, FArchive& Ar, FCbBufferAllocator Allo
 				FIoHash ObjectHash = HashField.AsObjectAttachment();
 				if (HashField.HasError() || Object.GetHash() != ObjectHash)
 				{
+					UE_LOG(LogZenSerialization, Warning, TEXT("package attachment has malformed/invalid object hash field"));
 					Ar.SetError();
 					return false;
 				}
@@ -152,9 +158,9 @@ struct CbAttachmentEntry
 
 struct CbAttachmentReferenceHeader
 {
-	uint64_t PayloadByteOffset = 0;
-	uint64_t PayloadByteSize = ~0u;
-	uint16_t AbsolutePathLength = 0;
+	uint64 PayloadByteOffset = 0;
+	uint64 PayloadByteSize = ~0u;
+	uint16 AbsolutePathLength = 0;
 
 	// This header will be followed by UTF8 encoded absolute path to backing file
 };
@@ -279,10 +285,18 @@ bool TryLoadCbPackage(FCbPackage& Package, FArchive& Ar, FCbBufferAllocator Allo
 		{
 			// Marshal local reference - a "pointer" to the chunk backing file
 
-			check(AttachmentData.GetSize() >= sizeof(CbAttachmentReferenceHeader));
+			if (AttachmentData.GetSize() < sizeof(CbAttachmentReferenceHeader))
+			{
+				// We did not get the full package data - treat it as a malformed package
+				UE_LOG(LogZenSerialization, Warning, TEXT("package payload is not large enough while trying to read local ref header, need %" UINT64_FMT " bytes, but only %" UINT64_FMT " available."), sizeof(CbAttachmentReferenceHeader), AttachmentData.GetSize());
+				Ar.SetError();
+				return false;
+			}
 
 			if (!(Entry.Flags & CbAttachmentEntry::IsCompressed))
 			{
+				// We only support compressed attachments - treat it as an error
+				UE_LOG(LogZenSerialization, Error, TEXT("attachment data is not compressed"));
 				Ar.SetError();
 				return false;
 			}
@@ -290,13 +304,21 @@ bool TryLoadCbPackage(FCbPackage& Package, FArchive& Ar, FCbBufferAllocator Allo
 			const CbAttachmentReferenceHeader* AttachRefHdr = reinterpret_cast<CbAttachmentReferenceHeader*>(AttachmentData.GetData());
 			const UTF8CHAR* PathPointer = reinterpret_cast<const UTF8CHAR*>(AttachRefHdr + 1);
 
-			check(AttachmentData.GetSize() >= (sizeof(CbAttachmentReferenceHeader) + AttachRefHdr->AbsolutePathLength));
+			const uint64 ExpectedHeaderSize = sizeof(CbAttachmentReferenceHeader) + AttachRefHdr->AbsolutePathLength;
+			if (AttachmentData.GetSize() < ExpectedHeaderSize)
+			{
+				// We did not get the full package data - treat it as a malformed package
+				UE_LOG(LogZenSerialization, Warning, TEXT("package payload is not large enough while trying to read local ref path, need %" UINT64_FMT " bytes, but only %" UINT64_FMT " available."), ExpectedHeaderSize, AttachmentData.GetSize());
+				Ar.SetError();
+				return false;
+			}
 
 			FString Path(FUtf8StringView(PathPointer, static_cast<uint32_t>(AttachRefHdr->AbsolutePathLength)));
 			TUniquePtr<FArchive> ChunkReader(IFileManager::Get().CreateFileReader(*Path, FILEREAD_Silent));
 			if (!ChunkReader)
 			{
 				// File does not exist on disk - treat it as missing
+				UE_LOG(LogZenSerialization, Warning, TEXT("unable to read local file '%s', treating it as a missing attachment"), *Path);
 				++Index;
 				continue;
 			}
@@ -304,6 +326,7 @@ bool TryLoadCbPackage(FCbPackage& Package, FArchive& Ar, FCbBufferAllocator Allo
 			if (!CompBuf)
 			{
 				// File has wrong format - treat it as an error
+				UE_LOG(LogZenSerialization, Error, TEXT("local file '%s' format is not compressed"), *Path);
 				Ar.SetError();
 				return false;
 			}
