Fix capitalization of logging in ZenSerialization.cpp

diff --git a/Engine/Source/Developer/Zen/Internal/ZenSerialization.cpp b/Engine/Source/Developer/Zen/Internal/ZenSerialization.cpp
--- a/Engine/Source/Developer/Zen/Internal/ZenSerialization.cpp
+++ b/Engine/Source/Developer/Zen/Internal/ZenSerialization.cpp
@@ -71,7 +71,7 @@ bool TryLoadCbPackage(FCbPackage& Package, FArchive& Ar, FCbBufferAllocator Allo
 		FCbField ValueField = LoadCompactBinary(Ar, StackAllocator);
 		if (!ValueField)
 		{
-			UE_LOG(LogZenSerialization, Warning, TEXT("package is malformed, can't read compact binary data"));
+			UE_LOG(LogZenSerialization, Warning, TEXT("Package is malformed, can't read compact binary data"));
 			Ar.SetError();
 			return false;
 		}
@@ -89,7 +89,7 @@ bool TryLoadCbPackage(FCbPackage& Package, FArchive& Ar, FCbBufferAllocator Allo
 				const FIoHash& Hash = HashField.AsAttachment();
 				if (HashField.HasError() || FIoHash::HashBuffer(Buffer) != Hash)
 				{
-					UE_LOG(LogZenSerialization, Warning, TEXT("package attachment has malformed/invalid attachment hash field"));
+					UE_LOG(LogZenSerialization, Warning, TEXT("Package attachment has malformed/invalid attachment hash field"));
 					Ar.SetError();
 					return false;
 				}
@@ -108,7 +108,7 @@ bool TryLoadCbPackage(FCbPackage& Package, FArchive& Ar, FCbBufferAllocator Allo
 			FCbObject Object = ValueField.AsObject();
 			if (ValueField.HasError())
 			{
-				UE_LOG(LogZenSerialization, Warning, TEXT("package attachment has malformed/invalid object field"));
+				UE_LOG(LogZenSerialization, Warning, TEXT("Package attachment has malformed/invalid object field"));
 				Ar.SetError();
 				return false;
 			}
@@ -119,7 +119,7 @@ bool TryLoadCbPackage(FCbPackage& Package, FArchive& Ar, FCbBufferAllocator Allo
 				FIoHash ObjectHash = HashField.AsObjectAttachment();
 				if (HashField.HasError() || Object.GetHash() != ObjectHash)
 				{
-					UE_LOG(LogZenSerialization, Warning, TEXT("package attachment has malformed/invalid object hash field"));
+					UE_LOG(LogZenSerialization, Warning, TEXT("Package attachment has malformed/invalid object hash field"));
 					Ar.SetError();
 					return false;
 				}
@@ -288,7 +288,7 @@ bool TryLoadCbPackage(FCbPackage& Package, FArchive& Ar, FCbBufferAllocator Allo
 			if (AttachmentData.GetSize() < sizeof(CbAttachmentReferenceHeader))
 			{
 				// We did not get the full package data - treat it as a malformed package
-				UE_LOG(LogZenSerialization, Warning, TEXT("package payload is not large enough while trying to read local ref header, need %" UINT64_FMT " bytes, but only %" UINT64_FMT " available."), sizeof(CbAttachmentReferenceHeader), AttachmentData.GetSize());
+				UE_LOG(LogZenSerialization, Warning, TEXT("Package payload is not large enough while trying to read local ref header, need %" UINT64_FMT " bytes, but only %" UINT64_FMT " available."), sizeof(CbAttachmentReferenceHeader), AttachmentData.GetSize());
 				Ar.SetError();
 				return false;
 			}
@@ -296,7 +296,7 @@ bool TryLoadCbPackage(FCbPackage& Package, FArchive& Ar, FCbBufferAllocator Allo
 			if (!(Entry.Flags & CbAttachmentEntry::IsCompressed))
 			{
 				// We only support compressed attachments - treat it as an error
-				UE_LOG(LogZenSerialization, Error, TEXT("attachment data is not compressed"));
+				UE_LOG(LogZenSerialization, Error, TEXT("Attachment data is not compressed"));
 				Ar.SetError();
 				return false;
 			}
@@ -308,7 +308,7 @@ bool TryLoadCbPackage(FCbPackage& Package, FArchive& Ar, FCbBufferAllocator Allo
 			if (AttachmentData.GetSize() < ExpectedHeaderSize)
 			{
 				// We did not get the full package data - treat it as a malformed package
-				UE_LOG(LogZenSerialization, Warning, TEXT("package payload is not large enough while trying to read local ref path, need %" UINT64_FMT " bytes, but only %" UINT64_FMT " available."), ExpectedHeaderSize, AttachmentData.GetSize());
+				UE_LOG(LogZenSerialization, Warning, TEXT("Package payload is not large enough while trying to read local ref path, need %" UINT64_FMT " bytes, but only %" UINT64_FMT " available."), ExpectedHeaderSize, AttachmentData.GetSize());
 				Ar.SetError();
 				return false;
 			}
@@ -318,7 +318,7 @@ bool TryLoadCbPackage(FCbPackage& Package, FArchive& Ar, FCbBufferAllocator Allo
 			if (!ChunkReader)
 			{
 				// File does not exist on disk - treat it as missing
-				UE_LOG(LogZenSerialization, Warning, TEXT("unable to read local file '%s', treating it as a missing attachment"), *Path);
+				UE_LOG(LogZenSerialization, Warning, TEXT("Unable to read local file '%s', treating it as a missing attachment"), *Path);
 				++Index;
 				continue;
 			}
@@ -326,7 +326,7 @@ bool TryLoadCbPackage(FCbPackage& Package, FArchive& Ar, FCbBufferAllocator Allo
 			if (!CompBuf)
 			{
 				// File has wrong format - treat it as an error
-				UE_LOG(LogZenSerialization, Error, TEXT("local file '%s' format is not compressed"), *Path);
+				UE_LOG(LogZenSerialization, Error, TEXT("Local file '%s' format is not compressed"), *Path);
 				Ar.SetError();
 				return false;
 			}
