Zen server interface changes in support of:
1) Zen versioning and installation of newest zen version across workspaces
2) Zen dashboard functionality

diff --git a/Engine/Source/Developer/Zen/Internal/ZenServerInterface.cpp b/Engine/Source/Developer/Zen/Internal/ZenServerInterface.cpp
--- a/Engine/Source/Developer/Zen/Internal/ZenServerInterface.cpp
+++ b/Engine/Source/Developer/Zen/Internal/ZenServerInterface.cpp
@@ -6,17 +6,20 @@
 #include "ZenSerialization.h"
 #include "ZenServerHttp.h"
 
+#include "Analytics.h"
+#include "AnalyticsEventAttribute.h"
+#include "Async/Async.h"
 #include "Dom/JsonValue.h"
 #include "HAL/FileManager.h"
 #include "HAL/Platform.h"
 #include "HAL/PlatformMisc.h"
 #include "HAL/PlatformProcess.h"
 #include "HAL/PlatformTime.h"
-#include "Async/Async.h"
 #include "Logging/LogScopedCategoryAndVerbosityOverride.h"
 #include "Misc/App.h"
 #include "Misc/ConfigCacheIni.h"
 #include "Misc/FileHelper.h"
+#include "Misc/MonitoredProcess.h"
 #include "Misc/Paths.h"
 #include "Misc/PathViews.h"
 #include "Misc/ScopedSlowTask.h"
@@ -24,10 +27,12 @@
 #include "Serialization/CompactBinary.h"
 #include "Serialization/CompactBinarySerialization.h"
 #include "Serialization/CompactBinaryValidation.h"
-#include "String/LexFromString.h"
+#include "Serialization/JsonReader.h"
+#include "Serialization/JsonSerializer.h"
+#include "Serialization/JsonWriter.h"
 #include "SocketSubsystem.h"
-#include "Analytics.h"
-#include "AnalyticsEventAttribute.h"
+#include "String/LexFromString.h"
+#include "ZenVersion.h"
 
 #if PLATFORM_WINDOWS
 #	include "Windows/AllowWindowsPlatformTypes.h"
@@ -49,6 +54,91 @@ namespace UE::Zen
 
 DEFINE_LOG_CATEGORY_STATIC(LogZenServiceInstance, Log, All);
 
+static FZenVersion
+GetZenVersion(const FString& UtilityPath, const FString& ServicePath, const FString& VersionCachePath)
+{
+	IFileManager& FileManager = IFileManager::Get();
+	FDateTime UtilityExecutableModificationTime = FileManager.GetTimeStamp(*UtilityPath);
+	FDateTime ServiceExecutableModificationTime = FileManager.GetTimeStamp(*ServicePath);
+
+	auto GetFallbackVersion = [&UtilityExecutableModificationTime, &ServiceExecutableModificationTime]()
+	{
+		FZenVersion FallbackVersion;
+		if (UtilityExecutableModificationTime > ServiceExecutableModificationTime)
+		{
+			FallbackVersion.Details = UtilityExecutableModificationTime.ToString();
+			return FallbackVersion;
+		}
+		FallbackVersion.Details = ServiceExecutableModificationTime.ToString();
+		return FallbackVersion;
+	};
+
+	FDateTime VersionCacheModificationTime = FileManager.GetTimeStamp(*VersionCachePath);
+	FString VersionFileContents;
+	FZenVersion ComparableVersion;
+	if ((VersionCacheModificationTime < UtilityExecutableModificationTime) ||
+		(VersionCacheModificationTime < ServiceExecutableModificationTime) ||
+		!FFileHelper::LoadFileToString(VersionFileContents, *VersionCachePath) ||
+		!ComparableVersion.TryParse(*VersionFileContents))
+	{
+		FString AbsoluteUtilityPath = FPaths::ConvertRelativePathToFull(UtilityPath);
+		FMonitoredProcess MonitoredUtilityProcess(AbsoluteUtilityPath, TEXT("version --detailed"), FPaths::GetPath(UtilityPath), true);
+
+		bool bLaunched = MonitoredUtilityProcess.Launch();
+		checkf(bLaunched, TEXT("Failed to launch zen utility to gather version data: '%s'."), *UtilityPath);
+		if (!bLaunched)
+		{
+			return GetFallbackVersion();
+		}
+
+		while (MonitoredUtilityProcess.Update())
+		{
+			FPlatformProcess::Sleep(0.1f);
+			if (MonitoredUtilityProcess.GetDuration().GetTotalSeconds() > 10)
+			{
+				MonitoredUtilityProcess.Cancel(true);
+				checkf(false, TEXT("Cancelled launch of zen utility for gathering version data: '%s'."), *UtilityPath);
+				return GetFallbackVersion();
+			}
+		}
+
+		if (MonitoredUtilityProcess.GetReturnCode() != 0)
+		{
+			checkf(false, TEXT("Unexpected return code after launch of zen utility for gathering version data: '%s' (%d)."), *UtilityPath, MonitoredUtilityProcess.GetReturnCode());
+			return GetFallbackVersion();
+		}
+
+		FString VersionOutputString = MonitoredUtilityProcess.GetFullOutputWithoutDelegate().TrimStartAndEnd();
+		if (!ComparableVersion.TryParse(*VersionOutputString))
+		{
+			checkf(false, TEXT("Invalid version information after launch of zen utility for gathering version data: '%s' (%s)."), *UtilityPath, *VersionOutputString);
+			return GetFallbackVersion();
+		}
+
+		FFileHelper::SaveStringToFile(ComparableVersion.ToString(), *VersionCachePath);
+	}
+
+	return ComparableVersion;
+}
+
+static bool
+IsInstallVersionOutOfDate(const FString& InTreeUtilityPath, const FString& InstallUtilityPath, const FString& InTreeServicePath, const FString& InstallServicePath, FString& OutInTreeVersionCache, FString& OutInstallVersionCache)
+{
+	OutInTreeVersionCache = FPaths::Combine(FPaths::EngineSavedDir(), TEXT("Zen") TEXT("zen.version"));
+	OutInstallVersionCache = FPaths::SetExtension(InstallUtilityPath, TEXT("version"));
+
+	IFileManager& FileManager = IFileManager::Get();
+	if (!FileManager.FileExists(*InstallUtilityPath) || !FileManager.FileExists(*InstallServicePath))
+	{
+		return true;
+	}
+
+	FZenVersion InTreeVersion = GetZenVersion(InTreeUtilityPath, InTreeServicePath, OutInTreeVersionCache);
+	FZenVersion InstallVersion = GetZenVersion(InstallUtilityPath, InstallServicePath, OutInstallVersionCache);
+
+	return InstallVersion < InTreeVersion;
+}
+
 static bool
 AttemptFileCopyWithRetries(const TCHAR* Dst, const TCHAR* Src, double RetryDurationSeconds)
 {
@@ -644,6 +734,241 @@ DetermineCmdLineWithoutTransientComponents(const FServiceAutoLaunchSettings& InS
 	return Parms;
 }
 
+bool
+TryGetLocalServiceRunContext(FZenLocalServiceRunContext& OutContext)
+{
+	FString LocalServiceRunContextPath = FPaths::SetExtension(GetLocalServiceInstallPath(), TEXT(".runcontext"));
+	return OutContext.ReadFromJsonFile(*LocalServiceRunContextPath);
+}
+
+bool
+FZenLocalServiceRunContext::ReadFromJson(FJsonObject& JsonObject)
+{
+	Executable = JsonObject.Values.FindRef(TEXT("Executable"))->AsString();
+	CommandlineArguments = JsonObject.Values.FindRef(TEXT("CommandlineArguments"))->AsString();
+	WorkingDirectory = JsonObject.Values.FindRef(TEXT("WorkingDirectory"))->AsString();
+	DataPath = JsonObject.Values.FindRef(TEXT("DataPath"))->AsString();
+	bShowConsole = JsonObject.Values.FindRef(TEXT("ShowConsole"))->AsBool();
+	return true;
+}
+
+void
+FZenLocalServiceRunContext::WriteToJson(TJsonWriter<TCHAR, TPrettyJsonPrintPolicy<TCHAR>>& Writer) const
+{
+	Writer.WriteValue(TEXT("Executable"), Executable);
+	Writer.WriteValue(TEXT("CommandlineArguments"), CommandlineArguments);
+	Writer.WriteValue(TEXT("WorkingDirectory"), WorkingDirectory);
+	Writer.WriteValue(TEXT("DataPath"), DataPath);
+	Writer.WriteValue(TEXT("ShowConsole"), bShowConsole);
+}
+
+bool
+FZenLocalServiceRunContext::ReadFromJsonFile(const TCHAR* Filename)
+{
+	FString JsonText;
+	if (!FFileHelper::LoadFileToString(JsonText, Filename))
+	{
+		return false;
+	}
+
+	TSharedPtr<FJsonObject> JsonObject;
+	TSharedRef<TJsonReader<TCHAR>> Reader = TJsonReaderFactory<TCHAR>::Create(JsonText);
+	if (!FJsonSerializer::Deserialize(Reader, JsonObject) || !JsonObject.IsValid())
+	{
+		return false;
+	}
+
+	return ReadFromJson(*JsonObject);
+}
+
+bool
+FZenLocalServiceRunContext::WriteToJsonFile(const TCHAR* Filename) const
+{
+	FString JsonTcharText;
+	TSharedRef<TJsonWriter<TCHAR, TPrettyJsonPrintPolicy<TCHAR>>> Writer = TJsonWriterFactory<TCHAR, TPrettyJsonPrintPolicy<TCHAR>>::Create(&JsonTcharText);
+	Writer->WriteObjectStart();
+	WriteToJson(*Writer);
+	Writer->WriteObjectEnd();
+	Writer->Close();
+
+	if (!FFileHelper::SaveStringToFile(JsonTcharText, Filename))
+	{
+		return false;
+	}
+
+	return true;
+}
+
+bool
+IsLocalServiceRunning(const TCHAR* DataPath, uint16* OutPort)
+{
+	// TODO: Instead of using the lock file, this could use the shared memory system state named "Global\ZenMap" (see zenserverprocess.{h,cpp} in zen codebase)
+	const FString LockFilePath = FPaths::Combine(DataPath, TEXT(".lock"));
+	if (IsLockFileLocked(*LockFilePath, true))
+	{
+		if (OutPort)
+		{
+			// If an instance is running with this data path, check if we can use it and what port it is on
+			*OutPort = 0;
+			FCbObject LockObject;
+			if (ReadCbLockFile(LockFilePath, LockObject))
+			{
+				bool bIsReady = LockObject["ready"].AsBool();
+				if (bIsReady)
+				{
+					*OutPort = LockObject["port"].AsUInt16();
+				}
+			}
+		}
+
+		return true;
+	}
+	return false;
+}
+
+FProcHandle
+StartLocalService(const FZenLocalServiceRunContext& Context, const TCHAR* TransientArgs)
+{
+	FString Parms = Context.GetCommandlineArguments();
+	if (TransientArgs)
+	{
+		Parms.Appendf(TEXT(" %s"), TransientArgs);
+	}
+
+	FProcHandle Proc;
+#if PLATFORM_WINDOWS
+	FString PlatformExecutable = Context.GetExecutable();
+	FPaths::MakePlatformFilename(PlatformExecutable);
+	FString PlatformWorkingDirectory = Context.GetWorkingDirectory();
+	FPaths::MakePlatformFilename(PlatformWorkingDirectory);
+	{
+		// Attempt non-elevated launch
+		STARTUPINFO StartupInfo = {
+			sizeof(STARTUPINFO),
+			NULL, NULL, NULL,
+			(::DWORD)CW_USEDEFAULT,
+			(::DWORD)CW_USEDEFAULT,
+			(::DWORD)CW_USEDEFAULT,
+			(::DWORD)CW_USEDEFAULT,
+			(::DWORD)0, (::DWORD)0, (::DWORD)0,
+			(::DWORD)STARTF_USESHOWWINDOW,
+			(::WORD)(Context.GetShowConsole() ? SW_SHOWMINNOACTIVE : SW_HIDE),
+			0, NULL,
+			HANDLE(nullptr),
+			HANDLE(nullptr),
+			HANDLE(nullptr)
+		};
+
+		FString CommandLine = FString::Printf(TEXT("\"%s\" %s"), *PlatformExecutable, *Parms);
+		PROCESS_INFORMATION ProcInfo;
+		if (CreateProcess(NULL, CommandLine.GetCharArray().GetData(), nullptr, nullptr, false, (::DWORD)(NORMAL_PRIORITY_CLASS | DETACHED_PROCESS), nullptr, PlatformWorkingDirectory.GetCharArray().GetData(), &StartupInfo, &ProcInfo))
+		{
+			::CloseHandle(ProcInfo.hThread);
+			Proc = FProcHandle(ProcInfo.hProcess);
+		}
+
+	}
+	if (!Proc.IsValid())
+	{
+		// Fall back to elevated launch
+		SHELLEXECUTEINFO ShellExecuteInfo;
+		ZeroMemory(&ShellExecuteInfo, sizeof(ShellExecuteInfo));
+		ShellExecuteInfo.cbSize = sizeof(ShellExecuteInfo);
+		ShellExecuteInfo.fMask = SEE_MASK_UNICODE | SEE_MASK_NOCLOSEPROCESS;
+		ShellExecuteInfo.lpFile = *PlatformExecutable;
+		ShellExecuteInfo.lpDirectory = *PlatformWorkingDirectory;
+		ShellExecuteInfo.lpVerb = TEXT("runas");
+		ShellExecuteInfo.nShow = Context.GetShowConsole() ? SW_SHOWMINNOACTIVE : SW_HIDE;
+		ShellExecuteInfo.lpParameters = *Parms;
+
+		if (ShellExecuteEx(&ShellExecuteInfo))
+		{
+			Proc = FProcHandle(ShellExecuteInfo.hProcess);
+		}
+	}
+#else
+	{
+		bool bLaunchDetached = true;
+		bool bLaunchHidden = true;
+		bool bLaunchReallyHidden = !Context.GetShowConsole();
+		uint32* OutProcessID = nullptr;
+		int32 PriorityModifier = 0;
+		void* PipeWriteChild = nullptr;
+		void* PipeReadChild = nullptr;
+		Proc = FPlatformProcess::CreateProc(
+			*Context.GetExecutable(),
+			*Parms,
+			bLaunchDetached,
+			bLaunchHidden,
+			bLaunchReallyHidden,
+			OutProcessID,
+			PriorityModifier,
+			*Context.GetWorkingDirectory(),
+			PipeWriteChild,
+			PipeReadChild);
+	}
+#endif
+	return Proc;
+}
+
+bool
+StopLocalService(const TCHAR* DataPath, double MaximumWaitDurationSeconds)
+{
+	uint16 CurrentPort = 0;
+	if (IsLocalServiceRunning(DataPath, &CurrentPort))
+	{
+		if (CurrentPort == 0)
+		{
+			return false;
+		}
+
+		RequestZenShutdownOnPort(CurrentPort);
+		return WaitForZenShutdown(*FPaths::Combine(DataPath, TEXT(".lock")), MaximumWaitDurationSeconds);
+	}
+	return true;
+}
+
+FString
+GetLocalServiceInstallPath()
+{
+	return FPaths::ConvertRelativePathToFull(FPaths::Combine(FPlatformProcess::ApplicationSettingsDir(), TEXT("Zen\\Install"),
+#if PLATFORM_WINDOWS
+		TEXT("zenserver.exe")
+#else
+		TEXT("zenserver")
+#endif
+		));
+}
+
+FString
+GetLocalServiceInstallVersion(bool bDetailed)
+{
+	FString InstallUtilityPath = GetLocalInstallUtilityPath();
+	FString InstallVersionCache = FPaths::SetExtension(InstallUtilityPath, TEXT("version"));
+
+	IFileManager& FileManager = IFileManager::Get();
+	if (!FileManager.FileExists(*InstallUtilityPath))
+	{
+		return FZenVersion().ToString(bDetailed);
+	}
+
+	FZenVersion InstallVersion = GetZenVersion(InstallUtilityPath, GetLocalServiceInstallPath(), InstallVersionCache);
+
+	return InstallVersion.ToString(bDetailed);
+}
+
+FString
+GetLocalInstallUtilityPath()
+{
+	return FPaths::ConvertRelativePathToFull(FPaths::Combine(FPlatformProcess::ApplicationSettingsDir(), TEXT("Zen\\Install"),
+#if PLATFORM_WINDOWS
+		TEXT("zen.exe")
+#else
+		TEXT("zen")
+#endif
+		));
+}
+
 static bool GIsDefaultServicePresent = false;
 
 FZenServiceInstance& GetDefaultServiceInstance()
@@ -794,50 +1119,49 @@ FZenServiceInstance::PromptUserToStopRunningServerInstance(const FString& Server
 FString
 FZenServiceInstance::ConditionalUpdateLocalInstall()
 {
-	FString InTreeFilePath = FPaths::ConvertRelativePathToFull(FPlatformProcess::GenerateApplicationPath(TEXT("zenserver"), EBuildConfiguration::Development));
-	FString InstallFilePath = FPaths::ConvertRelativePathToFull(FPaths::Combine(FPlatformProcess::ApplicationSettingsDir(), TEXT("Zen\\Install"), FString(FPathViews::GetCleanFilename(InTreeFilePath))));
+	FString InTreeUtilityPath = FPaths::ConvertRelativePathToFull(FPlatformProcess::GenerateApplicationPath(TEXT("zen"), EBuildConfiguration::Development));
+	FString InstallUtilityPath = GetLocalInstallUtilityPath();
+
+	FString InTreeServicePath = FPaths::ConvertRelativePathToFull(FPlatformProcess::GenerateApplicationPath(TEXT("zenserver"), EBuildConfiguration::Development));
+	FString InstallServicePath = GetLocalServiceInstallPath();
 
 	IFileManager& FileManager = IFileManager::Get();
-	FDateTime InTreeFileTime;
-	FDateTime InstallFileTime;
-	FileManager.GetTimeStampPair(*InTreeFilePath, *InstallFilePath, InTreeFileTime, InstallFileTime);
 
-	bool bMainExecutableUpdated = false;
-	if (InTreeFileTime > InstallFileTime)
+	bool bMainExecutablesUpdated = false;
+	FString InTreeVersionCache, InstallVersionCache;
+	if (IsInstallVersionOutOfDate(InTreeUtilityPath, InstallUtilityPath, InTreeServicePath, InstallServicePath, InTreeVersionCache, InstallVersionCache))
 	{
-		if (IsProcessActive(*InstallFilePath))
+		if (IsProcessActive(*InstallServicePath))
 		{
-			// TODO: Instead of using the lock file, this could use the shared memory system state named "Global\ZenMap" (see zenserverprocess.{h,cpp} in zen codebase)
-			FString LockFilePath = FPaths::Combine(Settings.SettingsVariant.Get<FServiceAutoLaunchSettings>().DataPath, TEXT(".lock"));
-			FCbObject LockObject;
-			if (ReadCbLockFile(LockFilePath, LockObject))
-			{
-				uint16 RunningPort = LockObject["port"].AsUInt16(0);
-				if (RunningPort != 0)
-				{
-					RequestZenShutdownOnPort(RunningPort);
+			FString DataPath = Settings.SettingsVariant.Get<FServiceAutoLaunchSettings>().DataPath;
+			StopLocalService(*DataPath);
 
-					WaitForZenShutdown(*LockFilePath, 5.0);
-				}
-			}
-
-			if (IsLockFileLocked(*LockFilePath))
+			if (IsLocalServiceRunning(*DataPath))
 			{
-				PromptUserToStopRunningServerInstance(InstallFilePath);
+				PromptUserToStopRunningServerInstance(InstallServicePath);
 			}
 		}
 
 		// Even after waiting for the lock file to be removed, the executable may have a period where it can't be overwritten as the process shuts down
 		// so any attempt to overwrite it should have some tolerance for retrying.
-		bMainExecutableUpdated = AttemptFileCopyWithRetries(*InstallFilePath, *InTreeFilePath, 5.0);
-		checkf(bMainExecutableUpdated, TEXT("Failed to copy zenserver to install location '%s'."), *InstallFilePath);
+		bool bServiceExecutableUpdated = AttemptFileCopyWithRetries(*InstallServicePath, *InTreeServicePath, 5.0);
+		checkf(bServiceExecutableUpdated, TEXT("Failed to copy zenserver to install location '%s'."), *InstallServicePath);
+
+		bool bUtilityExecutableUpdated = AttemptFileCopyWithRetries(*InstallUtilityPath, *InTreeUtilityPath, 5.0);
+		checkf(bUtilityExecutableUpdated, TEXT("Failed to copy zen to install location '%s'."), *InstallUtilityPath);
+
+		bMainExecutablesUpdated = (bServiceExecutableUpdated && bUtilityExecutableUpdated);
+		if (bMainExecutablesUpdated)
+		{
+			AttemptFileCopyWithRetries(*InstallVersionCache, *InTreeVersionCache, 1.0);
+		}
 	}
 
 #if PLATFORM_WINDOWS
-	FString InTreeSymbolFilePath = FPaths::ChangeExtension(InTreeFilePath, TEXT("pdb"));
-	FString InstallSymbolFilePath = FPaths::ChangeExtension(InstallFilePath, TEXT("pdb"));
+	FString InTreeSymbolFilePath = FPaths::SetExtension(InTreeServicePath, TEXT("pdb"));
+	FString InstallSymbolFilePath = FPaths::SetExtension(InstallServicePath, TEXT("pdb"));
 
-	if (FileManager.FileExists(*InTreeSymbolFilePath) && (bMainExecutableUpdated || !FileManager.FileExists(*InstallSymbolFilePath)))
+	if (FileManager.FileExists(*InTreeSymbolFilePath) && (bMainExecutablesUpdated || !FileManager.FileExists(*InstallSymbolFilePath)))
 	{
 		AttemptFileCopyWithRetries(*InstallSymbolFilePath, *InTreeSymbolFilePath, 1.0);
 	}
@@ -846,13 +1170,12 @@ FZenServiceInstance::ConditionalUpdateLocalInstall()
 	FString InTreeCrashpadHandlerFilePath = FPaths::ConvertRelativePathToFull(FPlatformProcess::GenerateApplicationPath(TEXT("crashpad_handler"), EBuildConfiguration::Development));
 	FString InstallCrashpadHandlerFilePath = FPaths::ConvertRelativePathToFull(FPaths::Combine(FPlatformProcess::ApplicationSettingsDir(), TEXT("Zen\\Install"), FString(FPathViews::GetCleanFilename(InTreeCrashpadHandlerFilePath))));
 
-	if (FileManager.FileExists(*InTreeCrashpadHandlerFilePath) && (bMainExecutableUpdated || !FileManager.FileExists(*InstallCrashpadHandlerFilePath)))
+	if (FileManager.FileExists(*InTreeCrashpadHandlerFilePath) && (bMainExecutablesUpdated || !FileManager.FileExists(*InstallCrashpadHandlerFilePath)))
 	{
 		AttemptFileCopyWithRetries(*InstallCrashpadHandlerFilePath, *InTreeCrashpadHandlerFilePath, 1.0);
 	}
 
-	FPaths::MakePlatformFilename(InstallFilePath);
-	return InstallFilePath;
+	return InstallServicePath;
 }
 
 bool
@@ -861,7 +1184,7 @@ FZenServiceInstance::AutoLaunch(const FServiceAutoLaunchSettings& InSettings, FS
 	int16 DesiredPort = InSettings.DesiredPort;
 	IFileManager& FileManager = IFileManager::Get();
 	const FString LockFilePath = FPaths::Combine(InSettings.DataPath, TEXT(".lock"));
-	const FString ExecutionContextFilePath = FPaths::Combine(InSettings.DataPath, TEXT(".runcontext"));
+	const FString ExecutionContextFilePath = FPaths::SetExtension(ExecutablePath, TEXT(".runcontext"));
 
 	FString WorkingDirectory = FPaths::GetPath(ExecutablePath);
 
@@ -882,9 +1205,16 @@ FZenServiceInstance::AutoLaunch(const FServiceAutoLaunchSettings& InSettings, FS
 		}
 
 		bool bCurrentInstanceUsable = false;
-		FString DesiredExecutionContext = FString::Printf(TEXT("%s %s") LINE_TERMINATOR TEXT("%s"), *ExecutablePath, *DetermineCmdLineWithoutTransientComponents(InSettings, CurrentPort), *WorkingDirectory);
-		FString CurrentExecutionContext;
-		if (FFileHelper::LoadFileToString(CurrentExecutionContext, *ExecutionContextFilePath) && (DesiredExecutionContext == CurrentExecutionContext))
+
+		FZenLocalServiceRunContext DesiredRunContext;
+		DesiredRunContext.Executable = ExecutablePath;
+		DesiredRunContext.CommandlineArguments = DetermineCmdLineWithoutTransientComponents(InSettings, CurrentPort);
+		DesiredRunContext.WorkingDirectory = WorkingDirectory;
+		DesiredRunContext.DataPath = InSettings.DataPath;
+		DesiredRunContext.bShowConsole = InSettings.bShowConsole;
+
+		FZenLocalServiceRunContext CurrentRunContext;
+		if (CurrentRunContext.ReadFromJsonFile(*ExecutionContextFilePath) && (DesiredRunContext == CurrentRunContext))
 		{
 			DesiredPort = CurrentPort;
 			bReUsingExistingInstance = true;
@@ -909,86 +1239,24 @@ FZenServiceInstance::AutoLaunch(const FServiceAutoLaunchSettings& InSettings, FS
 	if (InSettings.bLimitProcessLifetime || !bProcessIsLive)
 	{
 		FString ParmsWithoutTransients = DetermineCmdLineWithoutTransientComponents(InSettings, DesiredPort);
-		FString Parms = ParmsWithoutTransients;
+		FString TransientParms;
 
 		if (InSettings.bLimitProcessLifetime)
 		{
-			Parms.Appendf(TEXT(" --owner-pid %d"),
-				FPlatformProcess::GetCurrentProcessId());
+			TransientParms.Appendf(TEXT("--owner-pid %d"), FPlatformProcess::GetCurrentProcessId());
 		}
 
-		FProcHandle Proc;
-#if PLATFORM_WINDOWS
-		{
-			// Attempt non-elevated launch
-			STARTUPINFO StartupInfo = {
-				sizeof(STARTUPINFO),
-				NULL, NULL, NULL,
-				(::DWORD)CW_USEDEFAULT,
-				(::DWORD)CW_USEDEFAULT,
-				(::DWORD)CW_USEDEFAULT,
-				(::DWORD)CW_USEDEFAULT,
-				(::DWORD)0, (::DWORD)0, (::DWORD)0,
-				(::DWORD)STARTF_USESHOWWINDOW,
-				(::WORD)(InSettings.bShowConsole ? SW_SHOWMINNOACTIVE : SW_HIDE),
-				0, NULL,
-				HANDLE(nullptr),
-				HANDLE(nullptr),
-				HANDLE(nullptr)
-			};
-
-			FString CommandLine = FString::Printf(TEXT("\"%s\" %s"), *ExecutablePath, *Parms);
-			PROCESS_INFORMATION ProcInfo;
-			if (CreateProcess(NULL, CommandLine.GetCharArray().GetData(), nullptr, nullptr, false, (::DWORD)(NORMAL_PRIORITY_CLASS | DETACHED_PROCESS), nullptr, WorkingDirectory.GetCharArray().GetData(), &StartupInfo, &ProcInfo))
-			{
-				::CloseHandle(ProcInfo.hThread);
-				Proc = FProcHandle(ProcInfo.hProcess);
-			}
+		FZenLocalServiceRunContext EffectiveRunContext;
+		EffectiveRunContext.Executable = ExecutablePath;
+		EffectiveRunContext.CommandlineArguments = ParmsWithoutTransients;
+		EffectiveRunContext.WorkingDirectory = WorkingDirectory;
+		EffectiveRunContext.DataPath = InSettings.DataPath;
+		EffectiveRunContext.bShowConsole = InSettings.bShowConsole;
 
-		}
-		if (!Proc.IsValid())
-		{
-			// Fall back to elevated launch
-			SHELLEXECUTEINFO ShellExecuteInfo;
-			ZeroMemory(&ShellExecuteInfo, sizeof(ShellExecuteInfo));
-			ShellExecuteInfo.cbSize = sizeof(ShellExecuteInfo);
-			ShellExecuteInfo.fMask = SEE_MASK_UNICODE | SEE_MASK_NOCLOSEPROCESS;
-			ShellExecuteInfo.lpFile = *ExecutablePath;
-			ShellExecuteInfo.lpVerb = TEXT("runas");
-			ShellExecuteInfo.nShow = InSettings.bShowConsole ? SW_SHOWMINNOACTIVE : SW_HIDE;
-			ShellExecuteInfo.lpParameters = *Parms;
-
-			if (ShellExecuteEx(&ShellExecuteInfo))
-			{
-				Proc = FProcHandle(ShellExecuteInfo.hProcess);
-			}
-		}
-#else
-		{
-			bool bLaunchDetached = true;
-			bool bLaunchHidden = true;
-			bool bLaunchReallyHidden = !InSettings.bShowConsole;
-			uint32* OutProcessID = nullptr;
-			int32 PriorityModifier = 0;
-			void* PipeWriteChild = nullptr;
-			void* PipeReadChild = nullptr;
-			Proc = FPlatformProcess::CreateProc(
-				*ExecutablePath,
-				*Parms,
-				bLaunchDetached,
-				bLaunchHidden,
-				bLaunchReallyHidden,
-				OutProcessID,
-				PriorityModifier,
-				*WorkingDirectory,
-				PipeWriteChild,
-				PipeReadChild);
-		}
-#endif
+		FProcHandle Proc = StartLocalService(EffectiveRunContext, TransientParms.IsEmpty() ? nullptr : *TransientParms);
 		if (!bProcessIsLive)
 		{
-			FString UsedExecutionContext = FString::Printf(TEXT("%s %s") LINE_TERMINATOR TEXT("%s"), *ExecutablePath, *ParmsWithoutTransients, *WorkingDirectory);
-			FFileHelper::SaveStringToFile(UsedExecutionContext,*ExecutionContextFilePath);
+			EffectiveRunContext.WriteToJsonFile(*ExecutionContextFilePath);
 		}
 
 		bProcessIsLive = Proc.IsValid();
@@ -1025,9 +1293,9 @@ FZenServiceInstance::AutoLaunch(const FServiceAutoLaunchSettings& InSettings, FS
 			}
 
 			double ZenWaitDuration = FPlatformTime::ToSeconds64(FPlatformTime::Cycles64() - ZenWaitStartTime);
-			if (ZenWaitDuration < 5.0)
+			if (ZenWaitDuration < 10.0)
 			{
-				// Initial 5 second window of higher frequency checks
+				// Initial 10 second window of higher frequency checks
 				FPlatformProcess::Sleep(0.01f);
 			}
 			else
@@ -1060,7 +1328,7 @@ FZenServiceInstance::AutoLaunch(const FServiceAutoLaunchSettings& InSettings, FS
 					UE_LOG(LogZenServiceInstance, Display, TEXT("Waiting for ZenServer to be ready..."));
 					DurationPhase = EWaitDurationPhase::Medium;
 				}
-				else if (!FApp::IsUnattended() && ZenWaitDuration > 10.0 && (DurationPhase == EWaitDurationPhase::Medium))
+				else if (!FApp::IsUnattended() && ZenWaitDuration > 20.0 && (DurationPhase == EWaitDurationPhase::Medium))
 				{
 					FText ZenLongWaitPromptTitle = NSLOCTEXT("Zen", "Zen_LongWaitPromptTitle", "Wait for ZenServer?");
 					FText ZenLongWaitPromptText = NSLOCTEXT("Zen", "Zen_LongWaitPromptText", "ZenServer is taking a long time to launch. It may be performing maintenance. Keep waiting?");
@@ -1207,7 +1475,7 @@ FZenServiceInstance::GetStats(FZenStats& Stats)
 						UpstreamStats.EndPointStats.Push(EndPointStats);
 					}
 
-					FCbObjectView CASObjectView = RootObjectView["cas"].AsObjectView();
+					FCbObjectView CASObjectView = RootObjectView["cid"].AsObjectView();
 					FCbObjectView CASSizeObjectView = CASObjectView["size"].AsObjectView();
 
 					FZenCASSizeStats& CASSizeStats = Stats.CASStats.Size;
@@ -1227,6 +1495,106 @@ FZenServiceInstance::GetStats(FZenStats& Stats)
 	return Stats.IsValid;
 }
 
+bool 
+FZenServiceInstance::GetGCStatus(FGCStatus& Status)
+{
+	check(IsInGameThread());
+
+	// If we've already requested status and it is ready then grab it
+	if (GCStatusRequest.IsReady() == true )
+	{
+		LastGCStatus	 = GCStatusRequest.Get();
+		LastGCStatusTime = FPlatformTime::Cycles64();
+
+		GCStatusRequest.Reset();
+	}
+	
+	// Make a copy of the last updated status
+	if (LastGCStatus.IsSet())
+	{
+		Status = LastGCStatus.GetValue();
+	}
+
+	const uint64 CurrentTime = FPlatformTime::Cycles64();
+	constexpr double MinTimeBetweenRequestsInSeconds = 0.5;
+	const double DeltaTimeInSeconds = FPlatformTime::ToSeconds64(CurrentTime - LastGCStatusTime);
+
+	if (!GCStatusRequest.IsValid() && DeltaTimeInSeconds > MinTimeBetweenRequestsInSeconds)
+	{
+#if WITH_EDITOR
+		EAsyncExecution ThreadPool = EAsyncExecution::LargeThreadPool;
+#else
+		EAsyncExecution ThreadPool = EAsyncExecution::ThreadPool;
+#endif
+		if (!GCStatusHttpRequest.IsValid())
+		{
+			TStringBuilder<128> ZenDomain;
+			ZenDomain << HostName << TEXT(":") << Port;
+			GCStatusHttpRequest = MakePimpl<FZenHttpRequest>(ZenDomain.ToString(), false);
+		}
+
+		// We've not got any requests in flight and we've met a given time requirement for requests
+		GCStatusRequest = Async(ThreadPool, [this]
+			{
+				UE::Zen::FZenHttpRequest& Request = *GCStatusHttpRequest.Get();
+				Request.Reset();
+
+				TArray64<uint8> GetBuffer;
+				FZenHttpRequest::Result Result = Request.PerformBlockingDownload(TEXTVIEW("/admin/gc"), &GetBuffer, Zen::EContentType::CbObject);
+
+				TOptional<FGCStatus> GCStatus;
+
+				if (Result == Zen::FZenHttpRequest::Result::Success && Request.GetResponseCode() == 200)
+				{
+					FCbObjectView RootObjectView(GetBuffer.GetData());
+
+					GCStatus.Emplace();
+					GCStatus->Description = FString(RootObjectView["Status"].AsString());
+				}
+
+				return GCStatus;
+			});
+	}
+
+	return LastGCStatus.IsSet();
+}
+
+bool 
+FZenServiceInstance::RequestGC(const bool* OverrideCollectSmallObjects, const uint32* OverrideMaxCacheDuration)
+{
+	TStringBuilder<128> ZenDomain;
+	ZenDomain << HostName << TEXT(":") << Port;
+	UE::Zen::FZenHttpRequest Request(ZenDomain.ToString(), false);
+
+	TCHAR Separators[] = {TEXT('?'), TEXT('&')};
+	int32 SeparatorIndex = 0;
+	TStringBuilder<128> Query;
+	Query << TEXTVIEW("/admin/gc");
+
+	if (OverrideCollectSmallObjects)
+	{
+		Query << Separators[SeparatorIndex] << LexToString(*OverrideCollectSmallObjects);
+		SeparatorIndex = FMath::Min(SeparatorIndex + 1, (int32)UE_ARRAY_COUNT(Separators));
+	}
+
+	if (OverrideMaxCacheDuration)
+	{
+		Query << Separators[SeparatorIndex] << LexToString(*OverrideMaxCacheDuration);
+		SeparatorIndex = FMath::Min(SeparatorIndex + 1, (int32)UE_ARRAY_COUNT(Separators));
+	}
+
+	FZenHttpRequest::Result Result = Request.PerformBlockingPost(Query.ToString(), FMemoryView());
+
+	if (Result == Zen::FZenHttpRequest::Result::Success && Request.GetResponseCode() == 200)
+	{
+		FCbObjectView ResponseObject = FCbObjectView(Request.GetResponseBuffer().GetData());
+		FUtf8StringView ResponseStatus = ResponseObject["status"].AsString();
+
+		return (ResponseStatus == "Started") || (ResponseStatus == "Running");
+	}
+	return false;
+}
+
 void FZenServiceInstance::OnAnalyticsEvent(TArray<FAnalyticsEventAttribute>& Attributes )
 {
 	// Grab the Zen summary stats
diff --git a/Engine/Source/Developer/Zen/Internal/ZenVersion.cpp b/Engine/Source/Developer/Zen/Internal/ZenVersion.cpp
new file mode 100644
--- /dev/null
+++ b/Engine/Source/Developer/Zen/Internal/ZenVersion.cpp
@@ -0,0 +1,114 @@
+// Copyright Epic Games, Inc. All Rights Reserved.
+#include "ZenVersion.h"
+#include "Containers/StringFwd.h"
+#include "Misc/StringBuilder.h"
+
+namespace UE::Zen
+{
+
+void FZenVersion::Reset()
+{
+	MajorVersion = 0;
+	MinorVersion = 0;
+	PatchVersion = 0;
+	Details.Empty();
+}
+
+bool FZenVersion::TryParse(const TCHAR* InString)
+{
+	Reset();
+
+	if (InString == nullptr)
+	{
+		return false;
+	}
+
+	const TCHAR* CurrentSegment = InString;
+	uint32 CurrentSegmentIndex;
+	for (CurrentSegmentIndex = 0; CurrentSegmentIndex < 3; ++CurrentSegmentIndex)
+	{
+		TCHAR* End = nullptr;
+		uint64 ParsedValue = FCString::Strtoui64(CurrentSegment, &End, 10);
+
+		if (End == CurrentSegment)
+		{
+			// A segment with zero parsable characters is an error
+			break;
+		}
+
+		switch (CurrentSegmentIndex)
+		{
+		case 0:
+			MajorVersion = (uint32)ParsedValue;
+			break;
+		case 1:
+			MinorVersion = (uint32)ParsedValue;
+			break;
+		case 2:
+			PatchVersion = (uint32)ParsedValue;
+			break;
+		}
+
+		if (*End != TCHAR('.'))
+		{
+			CurrentSegment = End;
+			if (*End == TCHAR('-'))
+			{
+				Details = End + 1;
+				CurrentSegment = End + 1 + Details.Len();
+			}
+			++CurrentSegmentIndex;
+			break;
+		}
+		CurrentSegment = End + 1;
+	}
+
+	// Successful parsing is determined by whether we managed to parse at least 1 numeric segment
+	// and we've parsed to the end of the given string.  Details are optional.
+	return (CurrentSegmentIndex > 0) && (*CurrentSegment == TEXT('\0'));
+}
+
+FString FZenVersion::ToString(bool bDetailed) const
+{
+	TStringBuilder<64> Builder;
+	Builder << MajorVersion;
+	Builder.AppendChar(TEXT('.'));
+	Builder << MinorVersion;
+	Builder.AppendChar(TEXT('.'));
+	Builder << PatchVersion;
+	if (bDetailed && !Details.IsEmpty())
+	{
+		Builder.AppendChar(TEXT('-'));
+		Builder << Details;
+	}
+	return Builder.ToString();
+}
+
+bool FZenVersion::operator<(FZenVersion& Other) const
+{
+	if (MajorVersion < Other.MajorVersion)
+	{
+		return true;
+	}
+	else if (MajorVersion == Other.MajorVersion)
+	{
+		if (MinorVersion < Other.MinorVersion)
+		{
+			return true;
+		}
+		else if (MinorVersion == Other.MinorVersion)
+		{
+			if (PatchVersion < Other.PatchVersion)
+			{
+				return true;
+			}
+			else if (PatchVersion == Other.PatchVersion)
+			{
+				return Details < Other.Details;
+			}
+		}
+	}
+	return false;
+}
+
+} // namespace UE::Zen
diff --git a/Engine/Source/Developer/Zen/Internal/ZenVersion.h b/Engine/Source/Developer/Zen/Internal/ZenVersion.h
new file mode 100644
--- /dev/null
+++ b/Engine/Source/Developer/Zen/Internal/ZenVersion.h
@@ -0,0 +1,24 @@
+// Copyright Epic Games, Inc. All Rights Reserved.
+
+#pragma once
+
+#include "HAL/Platform.h"
+#include "Containers/UnrealString.h"
+
+namespace UE::Zen
+{
+
+struct FZenVersion
+{
+	uint32 MajorVersion = 0;
+	uint32 MinorVersion = 0;
+	uint32 PatchVersion = 0;
+	FString Details;
+
+	void Reset();
+	bool TryParse(const TCHAR* InString);
+	FString ToString(bool bDetailed = true) const;
+	bool operator<(FZenVersion& Other) const;
+};
+
+} // namespace UE::Zen
diff --git a/Engine/Source/Developer/Zen/Public/ZenServerInterface.h b/Engine/Source/Developer/Zen/Public/ZenServerInterface.h
--- a/Engine/Source/Developer/Zen/Public/ZenServerInterface.h
+++ b/Engine/Source/Developer/Zen/Public/ZenServerInterface.h
@@ -68,6 +68,53 @@ namespace UE::Zen
 {
 
 class FZenServiceInstance;
+class FZenLocalServiceRunContext;
+
+UE_API bool TryGetLocalServiceRunContext(FZenLocalServiceRunContext& OutContext);
+
+class FZenLocalServiceRunContext
+{
+public:
+	FString GetExecutable() const { return Executable; }
+	FString GetCommandlineArguments() const { return CommandlineArguments; }
+	FString GetWorkingDirectory() const { return WorkingDirectory; }
+	FString GetDataPath() const { return DataPath; }
+	bool GetShowConsole() const { return bShowConsole; }
+
+	friend inline bool operator==(FZenLocalServiceRunContext LHS, FZenLocalServiceRunContext RHS)
+	{
+		return (LHS.Executable == RHS.Executable) && (LHS.CommandlineArguments == RHS.CommandlineArguments) && (LHS.WorkingDirectory == RHS.WorkingDirectory) && (LHS.DataPath == RHS.DataPath);
+	}
+
+private:
+	FString Executable;
+	FString CommandlineArguments;
+	FString WorkingDirectory;
+	FString DataPath;
+	bool bShowConsole;
+
+	bool ReadFromJson(FJsonObject& JsonObject);
+	void WriteToJson(TJsonWriter<TCHAR, TPrettyJsonPrintPolicy<TCHAR>>& Writer) const;
+
+	bool ReadFromJsonFile(const TCHAR* Filename);
+	bool WriteToJsonFile(const TCHAR* Filename) const;
+
+	friend class FZenServiceInstance;
+	friend UE_API bool TryGetLocalServiceRunContext(FZenLocalServiceRunContext& OutContext);
+};
+
+UE_API bool IsLocalServiceRunning(const TCHAR* DataPath, uint16* OutPort = nullptr);
+UE_API FProcHandle StartLocalService(const FZenLocalServiceRunContext& Context, const TCHAR* TransientArgs = nullptr);
+UE_API bool StopLocalService(const TCHAR* DataPath, double MaximumWaitDurationSeconds = 5.0);
+
+UE_API FString GetLocalServiceInstallPath();
+UE_API FString GetLocalServiceInstallVersion(bool bDetailed = true);
+UE_API FString GetLocalInstallUtilityPath();
+
+struct FGCStatus
+{
+	FString Description;
+};
 
 /**
  * Type used to declare usage of a Zen server instance whether the shared default instance or a unique non-default instance.
@@ -125,11 +172,14 @@ public:
 	inline const TCHAR* GetHostName() const { return *HostName; }
 	inline uint16 GetPort() const { return Port; }
 	inline const FServiceSettings& GetServiceSettings() const { return Settings; }
-	UE_API bool GetStats(FZenStats& Stats);
 	UE_API bool IsServiceRunning();
 	UE_API bool IsServiceReady();
 	UE_API bool IsServiceRunningLocally() const { return bIsRunningLocally; }
 
+	UE_API bool GetStats(FZenStats& Stats);
+	UE_API bool GetGCStatus(FGCStatus& Status);
+	UE_API bool RequestGC(const bool* OverrideCollectSmallObjects = nullptr, const uint32* OverrideMaxCacheDuration = nullptr);
+
 	static UE_API uint16 GetAutoLaunchedPort();
 
 private:
@@ -140,6 +190,11 @@ private:
 	static bool AutoLaunch(const FServiceAutoLaunchSettings& InSettings, FString&& ExecutablePath, FString& OutHostName, uint16& OutPort);
 	void OnAnalyticsEvent(TArray<FAnalyticsEventAttribute>& Attributes);
 
+	mutable TPimplPtr<class FZenHttpRequest> GCStatusHttpRequest;
+	mutable TFuture<TOptional<FGCStatus>> GCStatusRequest;
+	mutable TOptional<FGCStatus> LastGCStatus;
+	mutable uint64 LastGCStatusTime = 0;
+
 	mutable TPimplPtr<class FZenHttpRequest> StatsHttpRequest;
 	mutable TFuture<FZenStats> StatsRequest;
 	mutable FZenStats LastStats;
