Move Zen module source files into Private folder instead of Internal.

diff --git a/Engine/Source/Developer/Zen/Internal/ZenClient.cpp b/Engine/Source/Developer/Zen/Internal/ZenClient.cpp
deleted file mode 100644
diff --git a/Engine/Source/Developer/Zen/Internal/ZenSerialization.cpp b/Engine/Source/Developer/Zen/Internal/ZenSerialization.cpp
deleted file mode 100644
diff --git a/Engine/Source/Developer/Zen/Internal/ZenServerHttp.cpp b/Engine/Source/Developer/Zen/Internal/ZenServerHttp.cpp
deleted file mode 100644
diff --git a/Engine/Source/Developer/Zen/Internal/ZenServerInterface.cpp b/Engine/Source/Developer/Zen/Internal/ZenServerInterface.cpp
deleted file mode 100644
diff --git a/Engine/Source/Developer/Zen/Internal/ZenVersion.cpp b/Engine/Source/Developer/Zen/Internal/ZenVersion.cpp
deleted file mode 100644
diff --git a/Engine/Source/Developer/Zen/Private/ZenClient.cpp b/Engine/Source/Developer/Zen/Private/ZenClient.cpp
new file mode 100644
--- /dev/null
+++ b/Engine/Source/Developer/Zen/Private/ZenClient.cpp
@@ -0,0 +1,913 @@
+// Copyright Epic Games, Inc. All Rights Reserved.
+
+#include "ZenClient.h"
+#include "Async/Async.h"
+#include "Containers/StringView.h"
+#include "HAL/Event.h"
+#include "Math/UnrealMathUtility.h"
+#include "Memory/MemoryView.h"
+#include "Misc/StringBuilder.h"
+#include "ProfilingDebugging/CountersTrace.h"
+#include "ProfilingDebugging/CpuProfilerTrace.h"
+#include "Serialization/LargeMemoryWriter.h"
+#include "Serialization/MemoryReader.h"
+#include "IPAddress.h"
+#include "Sockets.h"
+#include "SocketSubsystem.h"
+
+DEFINE_LOG_CATEGORY_STATIC(LogZenClient, Log, All);
+
+TRACE_DECLARE_INT_COUNTER(ZenClient_StreamedBytes, TEXT("ZenClient/StreamedBytes"));
+TRACE_DECLARE_FLOAT_COUNTER(ZenClient_StreamedSeconds, TEXT("ZenClient/StreamedSeconds"));
+
+namespace UE::Zen {
+
+////////////////////////////////////////////////////////////////////////////////
+bool SendData(FSocket& Socket, FMemoryView Data)
+{
+	while (Data.IsEmpty() == false)
+	{
+		const int32 BytesToSend = (int32)FMath::Clamp<uint64>(Data.GetSize(), 0, MAX_int32);
+		int32 BytesSent = 0;
+		if (Socket.Send((const uint8*)Data.GetData(), BytesToSend, BytesSent) == false)
+		{
+			return false;
+		}
+
+		Data += BytesSent;
+	}
+
+	return true;
+}
+
+void CreateHttpUpgradeRequest(const TArray<FString>& Protocols, FAnsiStringBuilderBase& OutRequest)
+{
+	const FAnsiStringView Endpoint = ANSITEXTVIEW("/zen");
+
+	OutRequest << ANSITEXTVIEW("GET ") << Endpoint << ANSITEXTVIEW(" HTTP/1.1\r\n");
+	OutRequest << ANSITEXTVIEW("Upgrade: websocket\r\n");
+	OutRequest << ANSITEXTVIEW("Connection: upgrade\r\n");
+	OutRequest << ANSITEXTVIEW("Sec-WebSocket-Key: dGhlIHNhbXBsZSBub25jZQ==\r\n");
+
+	if (Protocols.Num())
+	{
+		OutRequest << ANSITEXTVIEW("Sec-WebSocket-Protocol: ");
+
+		for (int32 Idx = 0, Count = Protocols.Num(); Idx < Count; ++Idx)
+		{
+			if (Idx > 0)
+			{
+				OutRequest << ANSITEXTVIEW(" ");
+			}
+
+			FTCHARToUTF8 Utf8(*Protocols[Idx]);
+			OutRequest << FAnsiStringView(Utf8.Get(), Utf8.Length());
+		}
+		
+		OutRequest << ANSITEXTVIEW("\r\n");
+	}
+
+	OutRequest << ANSITEXTVIEW("\r\n");
+}
+
+bool ReceieveHttpUpgradeResponse(FSocket& Socket, FAnsiStringBuilderBase& Response)
+{
+	static constexpr int32 MaxLineSize = 1024;
+	uint8 Buffer[MaxLineSize];
+
+	for(;;)
+	{
+		int32 BytesRead = 0;
+		Socket.Recv(Buffer, MaxLineSize, BytesRead, ESocketReceiveFlags::Peek);
+
+		FAnsiStringView BufferView(reinterpret_cast<const ANSICHAR*>(Buffer), BytesRead);
+
+		int32 LineEndIndex = 0;
+		if (BufferView.FindChar('\n', LineEndIndex) && BytesRead > LineEndIndex)
+		{
+			check(BufferView[LineEndIndex - 1] == '\r');
+
+			if (Socket.Recv(Buffer, LineEndIndex + 1, BytesRead, ESocketReceiveFlags::None) == false)
+			{
+				return false;
+			}
+
+			check(BytesRead == LineEndIndex + 1);
+			
+			FAnsiStringView Line = BufferView.Left(LineEndIndex - 1);
+			if (Line.Len() == 0)
+			{
+				return true;
+			}
+
+			Response << Line << ANSITEXTVIEW("\n");
+		}
+	}
+}
+
+bool ValidateHandshake(FAnsiStringView UpgradeResponse, FString& OutReason)
+{
+	int32 NewLineIndex = 0;
+	if (UpgradeResponse.FindChar('\n', NewLineIndex))
+	{
+		FAnsiStringView StatusLine = UpgradeResponse.Left(NewLineIndex);
+		UpgradeResponse.RightChopInline(NewLineIndex + 1);
+
+		if (StatusLine != ANSITEXTVIEW("HTTP/1.1 101 Switching Protocols"))
+		{
+			FUTF8ToTCHAR Str(StatusLine.GetData(), StatusLine.Len());
+			OutReason = Str.Get(); 
+
+			return false;
+		}
+	}
+
+	//TODO: Validate accept hash and protocols
+
+	return true;
+}
+
+////////////////////////////////////////////////////////////////////////////////
+enum class EZenWebSocketMessageType : uint8
+{
+	Invalid,
+	Notification,
+	Request,
+	StreamRequest,
+	Response,
+	StreamResponse,
+	StreamCompleteResponse,
+	Count
+};
+
+class FZenWebSocketMessage
+{
+	struct FHeader
+	{
+		static constexpr uint32_t ExpectedMagic = 0x7a776d68;
+
+		uint64 MessageSize = 0;
+		uint32 Magic = ExpectedMagic;
+		uint32 CorrelationId = 0;
+		uint32 StatusCode = 200;
+		EZenWebSocketMessageType MessageType = EZenWebSocketMessageType::Invalid;
+		uint8 Reserved[3] = {0};
+	};
+
+	static_assert(sizeof(FHeader) == 24); 
+
+public:
+	static constexpr uint64 HeaderSize = sizeof(FHeader);
+
+	FZenWebSocketMessage() = default;
+
+	EZenWebSocketMessageType MessageType() const { return MessageHeader.MessageType; }
+	void SetMessageType(EZenWebSocketMessageType InType) { MessageHeader.MessageType = InType; }
+
+	uint64 MessageSize() const { return MessageHeader.MessageSize; }
+	uint64 TotalSize() const { return HeaderSize + MessageHeader.MessageSize; }
+
+	uint32 CorrelationId() const { return MessageHeader.CorrelationId; }
+	void SetCorrelationId(uint32 Id) { MessageHeader.CorrelationId = Id; }
+
+	FCbPackage& Body() { return MessageBody; }
+	FCbPackage&& ConsumeBody() { return MoveTemp(MessageBody); }
+	void SetBody(FCbPackage&& Body) { MessageBody = MoveTemp(Body); }
+
+	bool TryLoadHeader(FMemoryView HeaderView);
+	void Save(FArchive& Ar);
+
+	static uint32 GenerateCorrelationId()
+	{
+		return NextCorrelationId.IncrementExchange();
+	}
+
+private:
+	static TAtomic<uint32> NextCorrelationId;
+
+	FHeader		MessageHeader;
+	FCbPackage	MessageBody;
+};
+
+bool FZenWebSocketMessage::TryLoadHeader(FMemoryView HeaderView)
+{
+	if (HeaderView.GetSize() < HeaderSize)
+	{
+		return false;
+	}
+
+	FMemory::Memcpy(&MessageHeader, HeaderView.GetData(), HeaderSize);
+
+	const bool bIsValidHeader =
+		MessageHeader.Magic == FHeader::ExpectedMagic &&
+		MessageHeader.StatusCode > 0 &&
+		static_cast<uint8>(MessageHeader.MessageType) > static_cast<uint8>(EZenWebSocketMessageType::Invalid) &&
+		static_cast<uint8>(MessageHeader.MessageType) < static_cast<uint8>(EZenWebSocketMessageType::Count);
+
+	return bIsValidHeader;
+}
+
+void FZenWebSocketMessage::Save(FArchive& Ar)
+{
+	Ar.Serialize(&MessageHeader, FZenWebSocketMessage::HeaderSize);
+	MessageBody.Save(Ar);
+	
+	const int64 End = Ar.Tell();
+
+	MessageHeader.MessageSize = End - int64(FZenWebSocketMessage::HeaderSize);
+	MessageHeader.StatusCode = 200;
+
+	Ar.Seek(0);
+	Ar.Serialize(&MessageHeader, FZenWebSocketMessage::HeaderSize);
+	Ar.Seek(End);
+}
+
+TAtomic<uint32> FZenWebSocketMessage::NextCorrelationId{1};
+
+////////////////////////////////////////////////////////////////////////////////
+enum class EParseMessageStatus : uint32
+{
+	Error,
+	Continue,
+	Done
+};
+
+class FZenMessageParser
+{
+public:
+	FZenMessageParser() = default;
+
+	EParseMessageStatus		ParseMessage(FMemoryView MessageData, uint64& OutBytesParsed, FString& OutReason);
+	void					Reset();
+	FZenWebSocketMessage&&	ConsumeMessage() { return MoveTemp(Message); }
+
+private:
+	FLargeMemoryWriter Ar;
+	FZenWebSocketMessage Message;
+};
+
+void FZenMessageParser::Reset()
+{
+	Ar.Seek(0);
+	Message = FZenWebSocketMessage();
+}
+
+EParseMessageStatus FZenMessageParser::ParseMessage(FMemoryView MessageData, uint64& OutBytesParsed, FString& OutReason)
+{
+	TRACE_CPUPROFILER_EVENT_SCOPE(FZenMessageParser::ParseMessage);
+
+	const uint64 PrevPos = Ar.Tell();
+
+	if (Ar.Tell() < FZenWebSocketMessage::HeaderSize)
+	{
+		FMemoryView HeaderView = MessageData.Left(FZenWebSocketMessage::HeaderSize - Ar.Tell());
+
+		Ar.Serialize(const_cast<void*>(HeaderView.GetData()), HeaderView.GetSize());
+		OutBytesParsed = Ar.Tell() - PrevPos;
+		
+		MessageData.RightChopInline(HeaderView.GetSize());
+
+		if (Ar.Tell() < FZenWebSocketMessage::HeaderSize)
+		{
+			return EParseMessageStatus::Continue;
+		}
+		
+		if (Message.TryLoadHeader(Ar.GetView()) == false)
+		{
+			OutReason = TEXT("Invalid websocket header");
+
+			return EParseMessageStatus::Error;
+		}
+
+		if (Message.MessageSize() == 0)
+		{
+			return EParseMessageStatus::Done;
+		}
+
+		//Ar.Reserve(Message.MessageSize());
+	}
+
+	if (MessageData.IsEmpty() == false)
+	{
+		const uint64 RemainingMessageSize = Message.MessageSize() - (Ar.Tell() - FZenWebSocketMessage::HeaderSize);
+		FMemoryView BodyView = MessageData.Left(RemainingMessageSize);
+
+		Ar.Serialize(const_cast<void*>(BodyView.GetData()), BodyView.GetSize());
+		OutBytesParsed = Ar.Tell() - PrevPos;
+	}
+
+	EParseMessageStatus Status = EParseMessageStatus::Continue;
+
+	if (Ar.Tell() == FZenWebSocketMessage::HeaderSize + Message.MessageSize())
+	{
+		Status = EParseMessageStatus::Done;
+
+		FMemoryView BodyView = Ar.GetView();
+		BodyView.RightChopInline(FZenWebSocketMessage::HeaderSize);
+		FMemoryReaderView BodyAr(BodyView);
+
+		FCbPackage Body;
+		{
+			TRACE_CPUPROFILER_EVENT_SCOPE(CbPackage::TryLoad);
+
+			if (Body.TryLoad(BodyAr))
+			{
+				Message.SetBody(MoveTemp(Body));
+			}
+			else
+			{
+				OutReason = TEXT("Invalid CbPackage");
+
+				Status = EParseMessageStatus::Error;
+			}
+		}
+	}
+
+	return Status;
+}
+
+////////////////////////////////////////////////////////////////////////////////
+class FZenClient final : public IZenClient
+{
+public:
+	FZenClient();
+	virtual ~FZenClient();
+
+	virtual bool Connect(FStringView Host, int32 Port) override;
+	virtual void Disconnect() override;
+	virtual bool IsConnected() const override { return bConnected; }
+
+	virtual bool SendRequest(FCbPackage&& Request, FOnStreamResponse&& OnResponse) override;
+	virtual bool SendRequest(FCbObject&& Request, FOnStreamResponse&& OnResponse) override;
+	virtual bool SendStreamRequest(FCbObject&& Request, FOnStreamResponse&& OnStreamResponse) override;
+
+private:
+	bool SendRequestMessage(FZenWebSocketMessage&& Request, FOnStreamResponse&& OnResponse);
+	void ReadMessageThreadEntry();
+	bool HandleMessage(FMemoryView MessageData, FString& OutReason);
+	void DispatchMessage(FZenWebSocketMessage&& Message);
+	void CloseConnection(FStringView Reason);
+
+	struct FPendingRequest
+	{
+		FPendingRequest(FOnStreamResponse&& Callback)
+			: OnStreamResponse(MoveTemp(Callback))
+			, StartTime(FPlatformTime::Seconds()) {}
+
+		FPendingRequest(const FPendingRequest&) = delete;
+		FPendingRequest& operator=(const FPendingRequest&) = delete;
+		
+		FOnStreamResponse OnStreamResponse;
+		FGraphEventRef LastCompletionEvent;
+		const double StartTime;
+	};
+
+	struct FDispatchStreamResponseTask
+	{
+		FDispatchStreamResponseTask(TSharedPtr<FPendingRequest> InRequest, FCbPackage&& InResponse, bool bInComplete)
+			: Request(InRequest), Response(MoveTemp(InResponse)), bComplete(bInComplete)
+		{
+			check(InRequest);
+		}
+
+		FORCEINLINE TStatId GetStatId() const { return TStatId(); }
+
+		ENamedThreads::Type GetDesiredThread() { return ENamedThreads::AnyBackgroundHiPriTask; }
+
+		static ESubsequentsMode::Type GetSubsequentsMode() { return ESubsequentsMode::TrackSubsequents; }
+
+		FORCEINLINE void DoTask(ENamedThreads::Type CurrentThread, const FGraphEventRef& CompletionGraphEvent)
+		{
+			Request->OnStreamResponse({ MoveTemp(Response), bComplete ? EZenStreamStatus::Completed :  EZenStreamStatus::Ok });
+		}
+
+		TSharedPtr<FPendingRequest> Request;
+		FCbPackage Response;
+		bool bComplete;
+	};
+
+	using FPendingRequests = TMap<uint32_t, TSharedPtr<FPendingRequest>>;
+
+	TUniquePtr<FSocket>				Socket;
+	TFuture<void>					ReadMessageThread;
+	TSharedPtr<FInternetAddr>		HostAddr;
+	TUniquePtr<FZenMessageParser>	Parser;
+	FPendingRequests				PendingRequests;
+	FCriticalSection				RequestsCriticalSection;
+	TArray<uint8>					ReadBuffer;
+	TAtomic<bool>					bConnected{false};
+	TAtomic<bool>					bStopRequested{false};
+};
+
+FZenClient::FZenClient()
+{
+	const int32 ReadBufferSize = 256 << 10;
+
+	Parser = MakeUnique<FZenMessageParser>();
+	ReadBuffer.SetNum(ReadBufferSize);
+}
+
+FZenClient::~FZenClient()
+{
+	Disconnect();
+}
+
+bool FZenClient::Connect(FStringView Host, int32 Port)
+{
+	ISocketSubsystem& SocketSubsystem = *ISocketSubsystem::Get();
+
+	TSharedPtr<FInternetAddr> Addr = SocketSubsystem.GetAddressFromString(FString(Host));
+	Addr->SetPort(Port);
+
+	Socket.Reset(SocketSubsystem.CreateSocket(NAME_Stream, TEXT("ZenClient"), Addr->GetProtocolType()));
+
+	if (Socket.IsValid() == false)
+	{
+		UE_LOG(LogZenClient, Error, TEXT("connect to server FAILED, reason 'failed to create socket'"));
+
+		return false;
+	}
+	
+	int32 NewReadBufferSize = 0;
+	if (Socket->SetReceiveBufferSize(ReadBuffer.Num(), NewReadBufferSize))
+	{
+		UE_LOG(LogZenClient, Log, TEXT("read buffer size set to %dB"), NewReadBufferSize);
+
+		ReadBuffer.SetNum(NewReadBufferSize);
+	}
+
+	//Socket->SetLoopBackFastPath(true);
+	Socket->SetNoDelay(true);
+
+	if (Socket->Connect(*Addr) == false)
+	{
+		UE_LOG(LogZenClient, Error, TEXT("connect to server FAILED"));
+
+		return false;
+	}
+
+	HostAddr = Addr;
+
+	TArray<FString> Protocols;
+	TAnsiStringBuilder<1024> UpgradeRequest;
+	CreateHttpUpgradeRequest(Protocols, UpgradeRequest);
+
+	if (SendData(*Socket, FMemoryView(UpgradeRequest.GetData(), UpgradeRequest.Len())) == false)
+	{
+		UE_LOG(LogZenClient, Error, TEXT("connect to server FAILED, reason 'failed to send handshake request'"));
+		
+		return false;
+	}
+
+	TAnsiStringBuilder<1024> UpgradeResponse;
+	if (ReceieveHttpUpgradeResponse(*Socket, UpgradeResponse) == false)
+	{
+		UE_LOG(LogZenClient, Error, TEXT("connect to server FAILED, reason 'failed to receive handshake request'"));
+
+		return false;
+	}
+
+	FString Reason;
+	if (ValidateHandshake(UpgradeResponse.ToView(), Reason) == false)
+	{
+		UE_LOG(LogZenClient, Error, TEXT("connect to server FAILED, reason '%s'"), *Reason);
+
+		return false;
+	}
+
+	const uint32 StackSize = 0; 
+	ReadMessageThread = AsyncThread([this] { return ReadMessageThreadEntry(); }, StackSize, TPri_AboveNormal);
+	
+	bConnected = true;
+
+	return true;
+}
+
+void FZenClient::Disconnect()
+{
+	CloseConnection(TEXTVIEW(""));
+}
+
+bool FZenClient::SendRequestMessage(FZenWebSocketMessage&& RequestMessage, FOnStreamResponse&& OnResponse)
+{
+	FLargeMemoryWriter Ar;
+	RequestMessage.Save(Ar);
+
+	TSharedPtr<FPendingRequest> PendingRequest = MakeShared<FPendingRequest>(MoveTemp(OnResponse));
+
+	{
+		FScopeLock _(&RequestsCriticalSection);
+		PendingRequests.Emplace(RequestMessage.CorrelationId(), MoveTemp(PendingRequest));
+	}
+
+	return SendData(*Socket, Ar.GetView());
+}
+
+bool FZenClient::SendRequest(FCbPackage&& Request, FOnStreamResponse&& OnResponse)
+{
+	FZenWebSocketMessage RequestMessage;
+
+	RequestMessage.SetMessageType(EZenWebSocketMessageType::Request);
+	RequestMessage.SetCorrelationId(FZenWebSocketMessage::GenerateCorrelationId());
+	RequestMessage.SetBody(MoveTemp(Request));
+
+	return SendRequestMessage(MoveTemp(RequestMessage), MoveTemp(OnResponse));
+}
+
+bool FZenClient::SendRequest(FCbObject&& Request, FOnStreamResponse&& OnResponse)
+{
+	FCbPackage Pkg;
+	Pkg.SetObject(MoveTemp(Request));
+
+	return SendRequest(MoveTemp(Pkg), MoveTemp(OnResponse));
+}
+
+bool FZenClient::SendStreamRequest(FCbObject&& Request, FOnStreamResponse&& OnStreamResponse)
+{
+	FZenWebSocketMessage StreamRequestMessage;
+
+	StreamRequestMessage.SetMessageType(EZenWebSocketMessageType::StreamRequest);
+	StreamRequestMessage.SetCorrelationId(FZenWebSocketMessage::GenerateCorrelationId());
+	StreamRequestMessage.SetBody(FCbPackage(MoveTemp(Request)));
+
+	return SendRequestMessage(MoveTemp(StreamRequestMessage), MoveTemp(OnStreamResponse));
+}
+
+void FZenClient::ReadMessageThreadEntry()
+{
+	FString Reason;
+
+	while(bStopRequested == false)
+	{
+		int32 BytesRead = 0;
+		{
+			if (Socket->Recv(ReadBuffer.GetData(), ReadBuffer.Num(), BytesRead) == false)
+			{
+				Reason = TEXT("receive message FAILED");
+				break;
+			}
+		}
+
+		if (HandleMessage(FMemoryView(ReadBuffer.GetData(), BytesRead), Reason) == false)
+		{
+			break;
+		}
+	}
+
+	if (bStopRequested == false)
+	{
+		CloseConnection(Reason);
+	}
+}
+
+bool FZenClient::HandleMessage(FMemoryView MessageData, FString& OutReason)
+{
+	TRACE_CPUPROFILER_EVENT_SCOPE(ZenClient::HandleMessage);
+
+	while (MessageData.IsEmpty() == false)
+	{
+		uint64 BytesParsed = 0;
+		const EParseMessageStatus Status = Parser->ParseMessage(MessageData, BytesParsed, OutReason);
+
+		MessageData.RightChopInline(BytesParsed);
+
+		if (Status == EParseMessageStatus::Continue)
+		{
+			check(MessageData.IsEmpty());
+
+			return true;
+		}
+
+		if (Status == EParseMessageStatus::Error)
+		{
+			return false;
+		}
+
+		DispatchMessage(Parser->ConsumeMessage());
+
+		Parser->Reset();
+	}
+
+	return true;
+}
+
+void FZenClient::DispatchMessage(FZenWebSocketMessage&& Message)
+{
+	TRACE_CPUPROFILER_EVENT_SCOPE(ZenClient::DispatchMessage);
+
+	const uint64 TotalMessageSize = Message.TotalSize();
+	
+	switch(Message.MessageType())
+	{
+	case EZenWebSocketMessageType::Notification:
+	{
+		check(false);
+		break;
+	}
+	case EZenWebSocketMessageType::Response:
+	case EZenWebSocketMessageType::StreamResponse:
+	case EZenWebSocketMessageType::StreamCompleteResponse:
+	{
+		const bool bIsStreamComplete =
+			Message.MessageType() == EZenWebSocketMessageType::StreamCompleteResponse || 
+			Message.MessageType() == EZenWebSocketMessageType::Response;
+
+		TGraphTask<FDispatchStreamResponseTask>* Task = nullptr;
+		double StartTime = 0.0;
+
+		{
+			FScopeLock _(&RequestsCriticalSection);
+			TSharedPtr<FPendingRequest> PendingRequest = PendingRequests.FindChecked(Message.CorrelationId());
+
+			if (bIsStreamComplete)
+			{
+				PendingRequests.Remove(Message.CorrelationId());
+				StartTime = PendingRequest->StartTime;
+			}
+
+			FGraphEventArray CompletionEvents;
+			if (PendingRequest->LastCompletionEvent.IsValid())
+			{
+				CompletionEvents.Add(PendingRequest->LastCompletionEvent);
+			}
+
+			Task = TGraphTask<FDispatchStreamResponseTask>::CreateTask(&CompletionEvents).ConstructAndHold(PendingRequest, Message.ConsumeBody(), bIsStreamComplete);
+			PendingRequest->LastCompletionEvent = Task->GetCompletionEvent();
+		}
+
+		Task->Unlock();
+
+		if (bIsStreamComplete)
+		{
+			TRACE_COUNTER_ADD(ZenClient_StreamedBytes, TotalMessageSize);
+			TRACE_COUNTER_ADD(ZenClient_StreamedSeconds, FPlatformTime::Seconds() - StartTime);
+		}
+
+		break;
+	}
+	default:
+		UE_LOG(LogZenClient, Warning, TEXT("dispatch message FAILED, reason 'invalid message type (%d)'"), Message.MessageType());
+		break;
+	}
+}
+
+void FZenClient::CloseConnection(FStringView Reason)
+{
+	if (bConnected.Exchange(false) == true)
+	{
+		UE_CLOG(Reason.IsEmpty() == false, LogZenClient, Warning, TEXT("closing connection, reason '%s'"), *FString(Reason));
+
+		check(PendingRequests.Num() == 0);
+
+		bStopRequested = true;
+
+		if (Socket)
+		{
+			Socket->Close();
+		}
+		
+		ReadMessageThread.Wait();
+		Socket.Reset();
+	}
+}
+
+TUniquePtr<IZenClient> IZenClient::Create()
+{
+	return MakeUnique<FZenClient>();
+}
+
+////////////////////////////////////////////////////////////////////////////////
+class FZenClientPool final : public IZenClientPool
+{
+public:
+	FZenClientPool() = default;
+	virtual ~FZenClientPool();
+
+	virtual bool Connect(FStringView Host, int32 Port, int32 PoolSize = 8) override;
+	virtual void Disconnect() override;
+
+	virtual inline bool IsConnected() const override { return bConnected; }
+
+	virtual bool SendRequest(FCbPackage&& Request, FOnStreamResponse&& OnResponse) override;
+	virtual bool SendRequest(FCbObject&& Request, FOnStreamResponse&& OnResponse) override;
+	virtual bool SendStreamRequest(FCbObject&& Request, FOnStreamResponse&& OnStreamResponse) override;
+
+private:
+
+#define UE_WITH_ZEN_NON_BLOCKING_POOL (1)
+
+#if UE_WITH_ZEN_NON_BLOCKING_POOL
+	struct FPoolEntry
+	{
+		std::atomic<uint8> bAllocated{0u};
+		FZenClient Client;
+	};
+#else
+	struct FPoolEntry
+	{
+		FZenClient Client;
+		FPoolEntry* Next = nullptr;
+	};
+#endif
+
+	using FPool = TArray<FPoolEntry>;
+
+	FPoolEntry* GetClient();
+	void ReleaseClient(FPoolEntry* Entry);
+
+	FPool				Pool;
+#if UE_WITH_ZEN_NON_BLOCKING_POOL
+	TAtomic<int32>		PoolEntryStartIndex{0};
+#else
+	FPoolEntry*			PoolHead = nullptr;
+	FPoolEntry*			PoolTail = nullptr;
+#endif
+	FCriticalSection	PoolCS;
+	FEventRef			PoolAvailable;
+	TAtomic<bool>		bConnected{false};
+};
+
+FZenClientPool::~FZenClientPool()
+{
+	Disconnect();
+}
+
+bool FZenClientPool::Connect(FStringView Host, int32 Port, int32 PoolSize)
+{
+	check(bConnected == false);
+
+	Pool.SetNum(PoolSize);
+
+	for (FPoolEntry& Entry : Pool)
+	{
+		if (Entry.Client.Connect(Host, Port) == false)
+		{
+			return false;		
+		}
+	}
+	
+	bConnected = true;
+
+#if UE_WITH_ZEN_NON_BLOCKING_POOL
+	for (FPoolEntry& Entry : Pool)
+	{
+		Entry.bAllocated = 0u;
+	}
+#else
+	PoolHead = PoolTail = &Pool[0];
+
+	for (int32 Index = 1; Index < PoolSize; ++Index)
+	{
+		Pool[Index].Next = PoolTail;
+		PoolTail = &Pool[Index];
+	}
+
+#endif
+
+	return bConnected;
+}
+
+void FZenClientPool::Disconnect()
+{
+	TRACE_CPUPROFILER_EVENT_SCOPE(ZenClientPool::Disconnect);
+
+	if (bConnected.Exchange(false) == true)
+	{
+		for (FPoolEntry& Entry : Pool)
+		{
+			Entry.Client.Disconnect();
+		}
+		Pool.Empty();
+	}
+}
+
+bool FZenClientPool::SendRequest(FCbPackage&& Request, FOnStreamResponse&& OnResponse)
+{
+	TRACE_CPUPROFILER_EVENT_SCOPE(ZenClientPool::SendRequest);
+
+	bool bResult = false;
+
+	if (FPoolEntry* Entry = GetClient())
+	{	
+		bResult = Entry->Client.SendRequest(MoveTemp(Request), MoveTemp(OnResponse));
+
+		ReleaseClient(Entry);
+	}
+
+	return bResult;
+}
+
+bool FZenClientPool::SendRequest(FCbObject&& Request, FOnStreamResponse&& OnResponse)
+{
+	FCbPackage Pkg;
+	Pkg.SetObject(MoveTemp(Request));
+
+	return SendRequest(MoveTemp(Pkg), MoveTemp(OnResponse));
+}
+
+bool FZenClientPool::SendStreamRequest(FCbObject&& Request, FOnStreamResponse&& OnStreamResponse)
+{
+	bool bResult = false;
+
+	if (FPoolEntry* Entry = GetClient())
+	{
+		bResult = Entry->Client.SendStreamRequest(MoveTemp(Request), MoveTemp(OnStreamResponse));
+		
+		ReleaseClient(Entry);
+	}
+
+	return bResult;
+}
+
+FZenClientPool::FPoolEntry* FZenClientPool::GetClient()
+{
+	TRACE_CPUPROFILER_EVENT_SCOPE(ZenClientPool::WaitForConnection);
+
+#if UE_WITH_ZEN_NON_BLOCKING_POOL
+
+	const int32 PoolSize = Pool.Num();
+	int32 EntryIndex = PoolEntryStartIndex.IncrementExchange() % PoolSize;
+	int32 SpinCount = 0;
+	
+	for (;;)
+	{
+		FPoolEntry& Entry = Pool[EntryIndex];
+			
+		if (Entry.bAllocated.load(std::memory_order_relaxed) == false)
+		{
+			uint8 Expected = 0u;
+			if (Entry.bAllocated.compare_exchange_strong(Expected, 1u))
+			{
+				return &Entry;
+			}
+		}
+
+		EntryIndex = (EntryIndex + 1) % PoolSize;
+		
+		if (++SpinCount >= PoolSize * 2)
+		{
+			SpinCount = 0;
+			FPlatformProcess::Sleep(0.01f);
+		}
+	}
+#else
+	for(;;)
+	{
+		{
+			FScopeLock _(&PoolCS);
+
+			if (FPoolEntry* Entry = PoolHead)
+			{
+				if (Entry->Next)
+				{
+					PoolHead = Entry->Next;
+				}
+				else
+				{
+					PoolHead = PoolTail = nullptr;
+				}
+
+				Entry->Next = nullptr;
+				return Entry;
+			}
+		}
+
+		PoolAvailable->Wait();
+	}
+#endif
+}
+
+void FZenClientPool::ReleaseClient(FZenClientPool::FPoolEntry* Entry)
+{
+#if UE_WITH_ZEN_NON_BLOCKING_POOL
+	uint8 Expected = 1u;
+	Entry->bAllocated.compare_exchange_strong(Expected, 0u);
+#else
+	check(Entry->Next == nullptr);
+
+	{
+		FScopeLock _(&PoolCS);
+		
+		if (PoolTail)
+		{
+			check(PoolTail->Next == nullptr);
+			PoolTail->Next = Entry;
+		}
+		else
+		{
+			PoolHead = Entry;
+		}
+		
+		PoolTail = Entry;
+	}
+	
+	PoolAvailable->Trigger();
+#endif
+}
+
+TUniquePtr<IZenClientPool> IZenClientPool::Create()
+{
+	return MakeUnique<FZenClientPool>();
+}
+
+} // namespace UE::Zen
diff --git a/Engine/Source/Developer/Zen/Private/ZenSerialization.cpp b/Engine/Source/Developer/Zen/Private/ZenSerialization.cpp
new file mode 100644
--- /dev/null
+++ b/Engine/Source/Developer/Zen/Private/ZenSerialization.cpp
@@ -0,0 +1,374 @@
+// Copyright Epic Games, Inc. All Rights Reserved.
+
+#include "ZenSerialization.h"
+#include "Compression/OodleDataCompression.h"
+#include "HAL/FileManager.h"
+#include "Serialization/CompactBinary.h"
+#include "Serialization/CompactBinaryPackage.h"
+#include "Serialization/CompactBinarySerialization.h"
+#include "Serialization/CompactBinaryValidation.h"
+#include "Serialization/CompactBinaryWriter.h"
+
+#if UE_WITH_ZEN
+
+DEFINE_LOG_CATEGORY_STATIC(LogZenSerialization, Log, All);
+
+namespace UE::Zen
+{
+
+namespace OpLog
+{
+void SaveCbAttachment(const FCbAttachment& Attachment, FCbWriter& Writer)
+{
+	if (Attachment.IsCompressedBinary())
+	{
+		Writer.AddBinary(Attachment.AsCompressedBinary().GetCompressed());
+		Writer.AddBinaryAttachment(Attachment.GetHash());
+	}
+	else if (Attachment.IsNull())
+	{
+		Writer.AddBinary(FMemoryView());
+	}
+	else
+	{
+		// NOTE: All attachments needs to be compressed
+		checkNoEntry();
+	}
+}
+
+void SaveCbPackage(const FCbPackage& Package, FCbWriter& Writer)
+{
+	if (const FCbObject& RootObject = Package.GetObject())
+	{
+		Writer.AddObject(RootObject);
+		Writer.AddObjectAttachment(Package.GetObjectHash());
+	}
+	for (const FCbAttachment& Attachment : Package.GetAttachments())
+	{
+		SaveCbAttachment(Attachment, Writer);
+	}
+	Writer.AddNull();
+}
+
+void SaveCbPackage(const FCbPackage& Package, FArchive& Ar)
+{
+	FCbWriter Writer;
+	SaveCbPackage(Package, Writer);
+	Writer.Save(Ar);
+}
+
+bool TryLoadCbPackage(FCbPackage& Package, FArchive& Ar, FCbBufferAllocator Allocator)
+{
+	uint8 StackBuffer[64];
+	const auto StackAllocator = [&Allocator, &StackBuffer](uint64 Size) -> FUniqueBuffer
+	{
+		return Size <= sizeof(StackBuffer) ? FUniqueBuffer::MakeView(StackBuffer, Size) : Allocator(Size);
+	};
+
+	Package = FCbPackage();
+	for (;;)
+	{
+		FCbField ValueField = LoadCompactBinary(Ar, StackAllocator);
+		if (!ValueField)
+		{
+			UE_LOG(LogZenSerialization, Warning, TEXT("Package is malformed, can't read compact binary data"));
+			Ar.SetError();
+			return false;
+		}
+		if (ValueField.IsNull())
+		{
+			return true;
+		}
+		else if (ValueField.IsBinary())
+		{
+			const FMemoryView View = ValueField.AsBinaryView();
+			if (View.GetSize() > 0)
+			{
+				FSharedBuffer Buffer = FSharedBuffer::MakeView(View, ValueField.GetOuterBuffer()).MakeOwned();
+				FCbField HashField = LoadCompactBinary(Ar, StackAllocator);
+				const FIoHash& Hash = HashField.AsAttachment();
+				if (HashField.HasError() || FIoHash::HashBuffer(Buffer) != Hash)
+				{
+					UE_LOG(LogZenSerialization, Warning, TEXT("Package attachment has malformed/invalid attachment hash field"));
+					Ar.SetError();
+					return false;
+				}
+				if (HashField.IsObjectAttachment())
+				{
+					Package.AddAttachment(FCbAttachment(FCbObject(MoveTemp(Buffer)), Hash));
+				}
+				else
+				{
+					Package.AddAttachment(FCbAttachment(FCompositeBuffer(MoveTemp(Buffer)), Hash));
+				}
+			}
+		}
+		else
+		{
+			FCbObject Object = ValueField.AsObject();
+			if (ValueField.HasError())
+			{
+				UE_LOG(LogZenSerialization, Warning, TEXT("Package attachment has malformed/invalid object field"));
+				Ar.SetError();
+				return false;
+			}
+
+			if (Object)
+			{
+				FCbField HashField = LoadCompactBinary(Ar, StackAllocator);
+				FIoHash ObjectHash = HashField.AsObjectAttachment();
+				if (HashField.HasError() || Object.GetHash() != ObjectHash)
+				{
+					UE_LOG(LogZenSerialization, Warning, TEXT("Package attachment has malformed/invalid object hash field"));
+					Ar.SetError();
+					return false;
+				}
+				Package.SetObject(Object, ObjectHash);
+			}
+		}
+	}
+}
+}
+
+namespace Http
+{
+struct CbPackageHeader
+{
+	uint32	HeaderMagic;
+	uint32	AttachmentCount;
+	uint32	Reserved1;
+	uint32	Reserved2;
+};
+
+struct CbAttachmentEntry
+{
+	uint64	AttachmentSize;
+	uint32	Flags;
+	FIoHash	AttachmentHash;
+
+	enum
+	{
+		IsCompressed = (1u << 0),	// Is marshaled using compressed buffer storage format
+		IsObject = (1u << 1),		// Is compact binary object
+		kIsError = (1u << 2),	    // Is error (compact binary formatted) object
+		kIsLocalRef = (1u << 3),	// Is "local reference"
+	};
+
+};
+
+struct CbAttachmentReferenceHeader
+{
+	uint64 PayloadByteOffset = 0;
+	uint64 PayloadByteSize = ~0u;
+	uint16 AbsolutePathLength = 0;
+
+	// This header will be followed by UTF8 encoded absolute path to backing file
+};
+
+void SaveCbPackage(const FCbPackage& Package, FArchive& Ar)
+{
+	TConstArrayView<FCbAttachment> Attachments = Package.GetAttachments();
+	const FCbObject& Object = Package.GetObject();
+	FCompressedBuffer ObjectBuffer = FCompressedBuffer::Compress(Object.GetBuffer(), FOodleDataCompression::ECompressor::NotSet, FOodleDataCompression::ECompressionLevel::None);
+
+	CbPackageHeader Hdr;
+	Hdr.HeaderMagic = kCbPkgMagic;
+	Hdr.AttachmentCount = Attachments.Num();
+	Hdr.Reserved1 = 0;
+	Hdr.Reserved2 = 0;
+
+	Ar.Serialize(&Hdr, sizeof Hdr);
+
+	// Root object metadata
+
+	{
+		CbAttachmentEntry Entry;
+		Entry.AttachmentHash = ObjectBuffer.GetRawHash();
+		Entry.AttachmentSize = ObjectBuffer.GetCompressedSize();
+		Entry.Flags = CbAttachmentEntry::IsObject | CbAttachmentEntry::IsCompressed;
+
+		Ar.Serialize(&Entry, sizeof Entry);
+	}
+
+	// Attachment metadata
+
+	for (const FCbAttachment& Attachment : Attachments)
+	{
+		CbAttachmentEntry Entry;
+		Entry.AttachmentHash = Attachment.GetHash();
+		Entry.Flags = 0;
+
+		if (Attachment.IsCompressedBinary())
+		{
+			Entry.AttachmentSize = Attachment.AsCompressedBinary().GetCompressedSize();
+			Entry.Flags |= CbAttachmentEntry::IsCompressed;
+		}
+		else if (Attachment.IsBinary())
+		{
+			Entry.AttachmentSize = Attachment.AsCompositeBinary().GetSize();
+		}
+		else if (Attachment.IsNull())
+		{
+			checkNoEntry();
+		}
+		else if (Attachment.IsObject())
+		{
+			Entry.AttachmentSize = Attachment.AsObject().GetSize();
+			Entry.Flags |= CbAttachmentEntry::IsObject;
+		}
+		else
+		{
+			checkNoEntry();
+		}
+
+		Ar.Serialize(&Entry, sizeof Entry);
+	}
+
+	// Root object
+
+	Ar << ObjectBuffer;
+
+	// Payloads back-to-back
+
+	for (const FCbAttachment& Attachment : Attachments)
+	{
+		if (Attachment.IsCompressedBinary())
+		{
+			FCompressedBuffer Payload = Attachment.AsCompressedBinary();
+			Ar << Payload;
+		}
+		else if (Attachment.IsBinary())
+		{
+			const FCompositeBuffer& Buffer = Attachment.AsCompositeBinary();
+			FSharedBuffer SharedBuffer = Buffer.ToShared();
+
+			Ar.Serialize((void*)SharedBuffer.GetData(), SharedBuffer.GetSize());
+		}
+		else if (Attachment.IsNull())
+		{
+			checkNoEntry();
+		}
+		else if (Attachment.IsObject())
+		{
+			checkNoEntry();
+		}
+		else
+		{
+			checkNoEntry();
+		}
+	}
+}
+
+bool TryLoadCbPackage(FCbPackage& Package, FArchive& Ar, FCbBufferAllocator Allocator)
+{
+	CbPackageHeader Hdr;
+	Ar.Serialize(&Hdr, sizeof Hdr);
+
+	if (Hdr.HeaderMagic != kCbPkgMagic)
+	{
+		return false;
+	}
+
+	TArray<CbAttachmentEntry> AttachmentEntries;
+	AttachmentEntries.SetNum(Hdr.AttachmentCount + 1);
+
+	Ar.Serialize(AttachmentEntries.GetData(), (Hdr.AttachmentCount + 1) * sizeof(CbAttachmentEntry));
+
+	int Index = 0;
+
+	for (const CbAttachmentEntry& Entry : AttachmentEntries)
+	{
+		FUniqueBuffer AttachmentData = FUniqueBuffer::Alloc(Entry.AttachmentSize);
+		Ar.Serialize(AttachmentData.GetData(), AttachmentData.GetSize());
+
+		if (Entry.Flags & CbAttachmentEntry::kIsLocalRef)
+		{
+			// Marshal local reference - a "pointer" to the chunk backing file
+
+			if (AttachmentData.GetSize() < sizeof(CbAttachmentReferenceHeader))
+			{
+				// We did not get the full package data - treat it as a malformed package
+				UE_LOG(LogZenSerialization, Warning, TEXT("Package payload is not large enough while trying to read local ref header, need %" UINT64_FMT " bytes, but only %" UINT64_FMT " available."), sizeof(CbAttachmentReferenceHeader), AttachmentData.GetSize());
+				Ar.SetError();
+				return false;
+			}
+
+			if (!(Entry.Flags & CbAttachmentEntry::IsCompressed))
+			{
+				// We only support compressed attachments - treat it as an error
+				UE_LOG(LogZenSerialization, Error, TEXT("Attachment data is not compressed"));
+				Ar.SetError();
+				return false;
+			}
+
+			const CbAttachmentReferenceHeader* AttachRefHdr = reinterpret_cast<CbAttachmentReferenceHeader*>(AttachmentData.GetData());
+			const UTF8CHAR* PathPointer = reinterpret_cast<const UTF8CHAR*>(AttachRefHdr + 1);
+
+			const uint64 ExpectedHeaderSize = sizeof(CbAttachmentReferenceHeader) + AttachRefHdr->AbsolutePathLength;
+			if (AttachmentData.GetSize() < ExpectedHeaderSize)
+			{
+				// We did not get the full package data - treat it as a malformed package
+				UE_LOG(LogZenSerialization, Warning, TEXT("Package payload is not large enough while trying to read local ref path, need %" UINT64_FMT " bytes, but only %" UINT64_FMT " available."), ExpectedHeaderSize, AttachmentData.GetSize());
+				Ar.SetError();
+				return false;
+			}
+
+			FString Path(FUtf8StringView(PathPointer, static_cast<uint32_t>(AttachRefHdr->AbsolutePathLength)));
+			TUniquePtr<FArchive> ChunkReader(IFileManager::Get().CreateFileReader(*Path, FILEREAD_Silent));
+			if (!ChunkReader)
+			{
+				// File does not exist on disk - treat it as missing
+				UE_LOG(LogZenSerialization, Warning, TEXT("Unable to read local file '%s', treating it as a missing attachment"), *Path);
+				++Index;
+				continue;
+			}
+			FCompressedBuffer CompBuf = FCompressedBuffer::Load(*ChunkReader);
+			if (!CompBuf)
+			{
+				// File has wrong format - treat it as an error
+				UE_LOG(LogZenSerialization, Error, TEXT("Local file '%s' format is not compressed"), *Path);
+				Ar.SetError();
+				return false;
+			}
+			FCbAttachment Attachment(MoveTemp(CompBuf));
+			Package.AddAttachment(Attachment);
+		}
+		else if (Entry.Flags & CbAttachmentEntry::IsCompressed)
+		{
+			FCompressedBuffer CompBuf(FCompressedBuffer::FromCompressed(AttachmentData.MoveToShared()));
+
+			if (Entry.Flags & CbAttachmentEntry::IsObject)
+			{
+				checkf(Index == 0, TEXT("Object attachments are not currently supported"));
+
+				Package.SetObject(FCbObject(CompBuf.Decompress()));
+			}
+			else
+			{
+				FCbAttachment Attachment(MoveTemp(CompBuf));
+				Package.AddAttachment(Attachment);
+			}
+		}
+		else /* not compressed */
+		{
+			if (Entry.Flags & CbAttachmentEntry::IsObject)
+			{
+				checkf(Index == 0, TEXT("Object attachments are not currently supported"));
+
+				Package.SetObject(FCbObject(AttachmentData.MoveToShared()));
+			}
+			else
+			{
+				FCbAttachment Attachment(AttachmentData.MoveToShared());
+				Package.AddAttachment(Attachment);
+			}
+		}
+
+		++Index;
+	}
+	return true;
+}
+}
+
+}
+#endif // UE_WITH_ZEN
diff --git a/Engine/Source/Developer/Zen/Private/ZenServerHttp.cpp b/Engine/Source/Developer/Zen/Private/ZenServerHttp.cpp
new file mode 100644
--- /dev/null
+++ b/Engine/Source/Developer/Zen/Private/ZenServerHttp.cpp
@@ -0,0 +1,734 @@
+// Copyright Epic Games, Inc. All Rights Reserved.
+
+#include "ZenServerHttp.h"
+
+#if UE_WITH_ZEN
+
+#include "ZenBackendUtils.h"
+#include "ZenSerialization.h"
+
+#if PLATFORM_MICROSOFT
+#	include "Microsoft/WindowsHWrapper.h"
+#	include "Microsoft/AllowMicrosoftPlatformTypes.h"
+#endif
+
+#if PLATFORM_WINDOWS
+#	include <mstcpip.h>
+#endif
+
+#if !defined(PLATFORM_CURL_INCLUDE)
+#include "curl/curl.h"
+#endif
+
+#if PLATFORM_MICROSOFT
+#	include "Microsoft/HideMicrosoftPlatformTypes.h"
+#endif
+
+#include "Logging/LogMacros.h"
+#include "Compression/CompressedBuffer.h"
+#include "Containers/StringFwd.h"
+#include "HAL/LowLevelMemTracker.h"
+#include "Memory/CompositeBuffer.h"
+#include "Serialization/CompactBinary.h"
+#include "Serialization/CompactBinaryPackage.h"
+#include "Serialization/CompactBinaryValidation.h"
+#include "Serialization/CompactBinaryWriter.h"
+#include "Serialization/LargeMemoryReader.h"
+#include "Serialization/LargeMemoryWriter.h"
+#include "HAL/PlatformProcess.h"
+#include "ProfilingDebugging/CpuProfilerTrace.h"
+#include "ZenSerialization.h"
+
+LLM_DEFINE_TAG(ZenDDC, NAME_None, TEXT("DDCBackend"));
+
+DEFINE_LOG_CATEGORY_STATIC(LogZenHttp, Log, All);
+
+namespace UE::Zen {
+
+#define UE_ZENDDC_BACKEND_WAIT_INTERVAL			0.01f
+#define UE_ZENDDC_HTTP_DEBUG					0
+
+	struct FZenHttpRequest::FStatics
+	{
+		static size_t StaticDebugCallback(CURL* Handle, curl_infotype DebugInfoType, char* DebugInfo, size_t DebugInfoSize, void* UserData);
+		static size_t StaticReadFn(void* Ptr, size_t SizeInBlocks, size_t BlockSizeInBytes, void* UserData);
+		static size_t StaticWriteHeaderFn(void* Ptr, size_t SizeInBlocks, size_t BlockSizeInBytes, void* UserData);
+		static size_t StaticWriteBodyFn(void* Ptr, size_t SizeInBlocks, size_t BlockSizeInBytes, void* UserData);
+		static size_t StaticSeekFn(void* UserData, curl_off_t Offset, int Origin);
+#if PLATFORM_WINDOWS
+		static int StaticSockoptFn(void* UserData, curl_socket_t CurlFd, curlsocktype Purpose);
+#endif //PLATFORM_WINDOWS
+	};
+
+	FZenHttpRequest::FZenHttpRequest(FStringView InDomain, bool bInLogErrors)
+		: bLogErrors(bInLogErrors)
+		, Domain(InDomain)
+	{
+		Curl = curl_easy_init();
+		Reset();
+	}
+
+	FZenHttpRequest::~FZenHttpRequest()
+	{
+		curl_easy_cleanup(Curl);
+	}
+
+	void FZenHttpRequest::Reset()
+	{
+		Headers.Reset();
+		ResponseHeader.Reset();
+		ResponseBuffer.Reset();
+		ResponseCode = 0;
+		bResponseFormatValid = true;
+		ReadDataView = nullptr;
+		WriteDataBufferPtr = nullptr;
+		WriteHeaderBufferPtr = nullptr;
+		BytesSent = 0;
+		BytesReceived = 0;
+		CurlResult = CURL_LAST;
+
+		curl_easy_reset(Curl);
+
+		// Options that are always set for all connections.
+		curl_easy_setopt(Curl, CURLOPT_CONNECTTIMEOUT, 30L);
+		curl_easy_setopt(Curl, CURLOPT_EXPECT_100_TIMEOUT_MS, 0);
+		curl_easy_setopt(Curl, CURLOPT_NOSIGNAL, 1L);
+		curl_easy_setopt(Curl, CURLOPT_DNS_CACHE_TIMEOUT, -1L); // Don't re-resolve names mid-session
+		curl_easy_setopt(Curl, CURLOPT_BUFFERSIZE, 256 * 1024L);
+		//curl_easy_setopt(Curl, CURLOPT_UPLOAD_BUFFERSIZE, 256 * 1024L);
+		// Response functions
+		curl_easy_setopt(Curl, CURLOPT_HEADERDATA, this);
+		curl_easy_setopt(Curl, CURLOPT_HEADERFUNCTION, &FZenHttpRequest::FStatics::StaticWriteHeaderFn);
+		curl_easy_setopt(Curl, CURLOPT_WRITEDATA, this);
+		curl_easy_setopt(Curl, CURLOPT_WRITEFUNCTION, &FZenHttpRequest::FStatics::StaticWriteBodyFn);
+		// Rewind method, handle special error case where request need to rewind data stream
+		curl_easy_setopt(Curl, CURLOPT_SEEKDATA, this);
+		curl_easy_setopt(Curl, CURLOPT_SEEKFUNCTION, &FZenHttpRequest::FStatics::StaticSeekFn);
+#if PLATFORM_WINDOWS
+		curl_easy_setopt(Curl, CURLOPT_SOCKOPTFUNCTION, &FZenHttpRequest::FStatics::StaticSockoptFn);
+#endif //PLATFORM_WINDOWS
+		// Debug hooks
+#if UE_ZENDDC_HTTP_DEBUG
+		curl_easy_setopt(Curl, CURLOPT_DEBUGDATA, this);
+		curl_easy_setopt(Curl, CURLOPT_DEBUGFUNCTION, &FZenHttpRequest::FStatics::StaticDebugCallback);
+		curl_easy_setopt(Curl, CURLOPT_VERBOSE, 1L);
+#endif
+	}
+
+	void FZenHttpRequest::Initialize(bool bInLogErrors)
+	{
+		bLogErrors = bInLogErrors;
+	}
+
+	void FZenHttpRequest::AddHeader(FStringView Header, FStringView Value)
+	{
+		TStringBuilder<128> Sb;
+		Sb << Header << TEXTVIEW(": ") << Value;
+		Headers.Emplace(Sb.ToString());
+	}
+
+	FZenHttpRequest::Result FZenHttpRequest::PerformBlockingPut(const TCHAR* Uri, const FCompositeBuffer& Buffer, EContentType ContentType)
+	{
+		uint64 ContentLength = 0u;
+
+		ContentLength = Buffer.GetSize();
+		curl_easy_setopt(Curl, CURLOPT_UPLOAD, 1L);
+		curl_easy_setopt(Curl, CURLOPT_INFILESIZE, ContentLength);
+		curl_easy_setopt(Curl, CURLOPT_READDATA, this);
+		curl_easy_setopt(Curl, CURLOPT_READFUNCTION, &FZenHttpRequest::FStatics::StaticReadFn);
+
+		AddHeader(TEXTVIEW("Content-Type"), GetMimeType(ContentType));
+
+		ReadDataView = &Buffer;
+
+		return PerformBlocking(Uri, RequestVerb::Put, ContentLength);
+	}
+
+	FZenHttpRequest::Result FZenHttpRequest::PerformBlockingPost(FStringView Uri, FCbObjectView Obj,
+		EContentType AcceptType)
+	{
+		FLargeMemoryWriter Out;
+		Obj.CopyTo(Out);
+
+		return PerformBlockingPost(Uri, Out.GetView(), EContentType::CbObject, AcceptType);
+	}
+
+	FZenHttpRequest::Result FZenHttpRequest::PerformBlockingPostPackage(FStringView Uri, const FCbPackage& Package, EContentType AcceptType)
+	{
+		FLargeMemoryWriter Out;
+		Http::SaveCbPackage(Package, Out);
+
+		return PerformBlockingPost(Uri, Out.GetView(), EContentType::CbPackage, AcceptType);
+	}
+
+	FZenHttpRequest::Result FZenHttpRequest::PerformRpc(FStringView Uri, FCbObjectView Request, FCbPackage &OutResponse)
+	{
+		return ParseRpcResponse(
+			PerformBlockingPost(Uri, Request, EContentType::CbPackage), OutResponse);
+	}
+		
+	FZenHttpRequest::Result FZenHttpRequest::PerformRpc(FStringView Uri, const FCbPackage& Request, FCbPackage& OutResponse)
+	{
+		return ParseRpcResponse(
+			PerformBlockingPostPackage(Uri, Request, EContentType::CbPackage), OutResponse);
+	}
+
+	FZenHttpRequest::Result FZenHttpRequest::ParseRpcResponse(FZenHttpRequest::Result ResultFromPost, FCbPackage& OutResponse)
+	{
+		if (ResultFromPost != Result::Success || !IsSuccessCode(ResponseCode))
+		{
+			return Result::Failed;
+		}
+
+		if (ResponseBuffer.Num())
+		{
+			{
+				FLargeMemoryReader Ar(ResponseBuffer.GetData(), ResponseBuffer.Num());
+				if (Http::TryLoadCbPackage(OutResponse, Ar))
+				{
+					return Result::Success;
+				}
+			}
+			FLargeMemoryReader Ar(ResponseBuffer.GetData(), ResponseBuffer.Num());
+			if (!OutResponse.TryLoad(Ar))
+			{
+				return Result::Failed;
+			}
+		}
+
+		return Result::Success;
+	}
+
+	FCbPackage FZenHttpRequest::GetResponseAsPackage() const
+	{
+		const TArray64<uint8>& Response = GetResponseBuffer();
+		FLargeMemoryReader Reader(Response.GetData(), Response.Num());
+
+		FCbPackage Package;
+		if (!Http::TryLoadCbPackage(Package, Reader))
+		{
+			return {};
+		}
+
+		return Package;
+	}
+
+	FZenHttpRequest::Result FZenHttpRequest::PerformBlockingPost(FStringView Uri, FMemoryView Payload,
+		EContentType ContentType, EContentType AcceptType)
+	{
+		uint64 ContentLength = 0u;
+
+		curl_easy_setopt(Curl, CURLOPT_POST, 1L);
+		curl_easy_setopt(Curl, CURLOPT_POSTFIELDSIZE, Payload.GetSize());
+		curl_easy_setopt(Curl, CURLOPT_READDATA, this);
+		curl_easy_setopt(Curl, CURLOPT_READFUNCTION, &FZenHttpRequest::FStatics::StaticReadFn);
+
+		AddHeader(TEXTVIEW("Content-Type"), GetMimeType(ContentType));
+		if (AcceptType != EContentType::UnknownContentType)
+		{
+			AddHeader(TEXTVIEW("Accept"), GetMimeType(EContentType::CbPackage));
+	}
+
+		ContentLength = Payload.GetSize();
+
+		FCompositeBuffer Buffer(FSharedBuffer::MakeView(Payload));
+		ReadDataView = &Buffer;
+
+		return PerformBlocking(Uri, RequestVerb::Post, ContentLength);
+	}
+
+	FZenHttpRequest::Result FZenHttpRequest::PerformBlockingDownload(FStringView Uri, TArray64<uint8>* Buffer, EContentType AcceptType)
+	{
+		curl_easy_setopt(Curl, CURLOPT_HTTPGET, 1L);
+		WriteDataBufferPtr = Buffer;
+		
+		AddHeader(TEXTVIEW("Accept"), GetMimeType(AcceptType));
+
+		return PerformBlocking(Uri, RequestVerb::Get, 0u);
+	}
+
+	FZenHttpRequest::Result FZenHttpRequest::PerformBlockingDownload(const TCHAR* Uri, FCbPackage& OutPackage)
+	{
+		curl_easy_setopt(Curl, CURLOPT_HTTPGET, 1L);
+		OutPackage.Reset();
+
+		AddHeader(TEXTVIEW("Accept"), GetMimeType(EContentType::CbPackage));
+
+		// TODO: When PackageBytes can be written in segments directly, set the WritePtr to the OutPackage and use that
+		TArray64<uint8> PackageBytes;
+		WriteDataBufferPtr = &PackageBytes;
+		Result LocalResult = PerformBlocking(Uri, RequestVerb::Get, 0u);
+		if (IsSuccessCode(ResponseCode))
+		{
+			FLargeMemoryReader Ar(PackageBytes.GetData(), PackageBytes.Num());
+			bResponseFormatValid = OutPackage.TryLoad(Ar);
+		}
+		return LocalResult;
+	}
+
+	FZenHttpRequest::Result FZenHttpRequest::PerformBlockingHead(FStringView Uri, EContentType AcceptType)
+	{
+		curl_easy_setopt(Curl, CURLOPT_NOBODY, 1L);
+		
+		AddHeader(TEXTVIEW("Accept"), GetMimeType(AcceptType));
+
+		return PerformBlocking(Uri, RequestVerb::Head, 0u);
+	}
+
+	FZenHttpRequest::Result FZenHttpRequest::PerformBlockingDelete(const FStringView Uri)
+	{
+		curl_easy_setopt(Curl, CURLOPT_CUSTOMREQUEST, "DELETE");
+
+		return PerformBlocking(Uri, RequestVerb::Delete, 0u);
+	}
+
+	static const char* GetSessionIdHeader() {
+		static FCbObjectId SessionId = FCbObjectId::NewObjectId();
+
+		static const char* HeaderString = [&] {
+			static TAnsiStringBuilder<64> SessionIdHeader;
+			SessionIdHeader << "UE-Session: " << SessionId;
+			return SessionIdHeader.GetData(); 
+		}();
+
+		return HeaderString;
+	}
+
+	static std::atomic<int> RequestId{1};
+
+	FZenHttpRequest::Result FZenHttpRequest::PerformBlocking(FStringView Uri, RequestVerb Verb, uint64 ContentLength)
+	{
+		LLM_SCOPE_BYTAG(ZenDDC);
+		// Strip any leading slashes because we compose the prefix and the suffix with a separating slash below
+		while (Uri.StartsWith(TEXT('/')))
+		{
+			Uri.RightChopInline(1);
+		}
+
+		TAnsiStringBuilder<32> RequestIdHeader;
+		RequestIdHeader << "UE-Request: " << RequestId++;
+
+		const char* CommonHeaders[] = {
+			GetSessionIdHeader(),
+			RequestIdHeader.GetData(),
+			// Strip any Expect: 100-Continue header since this just introduces latency
+			"Expect:",	
+			nullptr
+		};
+
+		TRACE_CPUPROFILER_EVENT_SCOPE(ZenHttp_CurlPerform);
+
+		// Setup request options
+		FString Url = FString::Printf(TEXT("%s/%s"), *Domain, *FString(Uri));
+		curl_easy_setopt(Curl, CURLOPT_URL, TCHAR_TO_ANSI(*Url));
+
+		// Setup response header buffer. If caller has not setup a response data buffer, use internal.
+		WriteHeaderBufferPtr = &ResponseHeader;
+		if (WriteDataBufferPtr == nullptr)
+		{
+			WriteDataBufferPtr = &ResponseBuffer;
+		}
+
+		if ((Verb != RequestVerb::Delete) && (Verb != RequestVerb::Get))
+		{
+			Headers.Add(FString::Printf(TEXT("Content-Length: %" UINT64_FMT), ContentLength));
+		}
+
+		// Build headers list
+		curl_slist* CurlHeaders = nullptr;
+		// Add common headers
+		for (uint8 i = 0; CommonHeaders[i] != nullptr; ++i)
+		{
+			CurlHeaders = curl_slist_append(CurlHeaders, CommonHeaders[i]);
+		}
+		// Setup added headers
+		for (const FString& Header : Headers)
+		{
+			CurlHeaders = curl_slist_append(CurlHeaders, TCHAR_TO_ANSI(*Header));
+		}
+		curl_easy_setopt(Curl, CURLOPT_HTTPHEADER, CurlHeaders);
+
+		// Shots fired!
+		CurlResult = curl_easy_perform(Curl);
+
+		// Get response code
+		bool bRedirected = false;
+		if (CURLE_OK == curl_easy_getinfo(Curl, CURLINFO_RESPONSE_CODE, &ResponseCode))
+		{
+			bRedirected = (ResponseCode >= 300 && ResponseCode < 400);
+		}
+
+		LogResult(CurlResult, *FString(Uri), Verb);
+
+		// Clean up
+		curl_slist_free_all(CurlHeaders);
+
+		return CurlResult == CURLE_OK ? Result::Success : Result::Failed;
+	}
+
+	/**
+	 * Attempts to find the header from the response. Returns false if header is not present.
+	 */
+	bool FZenHttpRequest::GetHeader(const ANSICHAR* Header, FString& OutValue) const
+	{
+		check(CurlResult != CURL_LAST);  // Cannot query headers before request is sent
+
+		const ANSICHAR* HeadersBuffer = (const ANSICHAR*)ResponseHeader.GetData();
+		size_t HeaderLen = strlen(Header);
+
+		// Find the header key in the (ANSI) response buffer. If not found we can exist immediately
+		if (const ANSICHAR* Found = strstr(HeadersBuffer, Header))
+		{
+			const ANSICHAR* Linebreak = strchr(Found, '\r');
+			const ANSICHAR* ValueStart = Found + HeaderLen + 2; //colon and space
+			const size_t ValueSize = Linebreak - ValueStart;
+			FUTF8ToTCHAR TCHARData(ValueStart, ValueSize);
+			OutValue = FString(TCHARData.Length(), TCHARData.Get());
+			return true;
+		}
+		return false;
+	}
+
+
+	void FZenHttpRequest::LogResult(long InResult, const TCHAR* Uri, RequestVerb Verb) const
+	{
+		CURLcode Result = (CURLcode)InResult;
+		if (Result == CURLE_OK)
+		{
+			bool bSuccess = false;
+			const TCHAR* VerbStr = nullptr;
+			FString AdditionalInfo;
+
+			const bool Is404 = (ResponseCode == 404);
+			const bool Is2xx = (ResponseCode >= 200) && (ResponseCode <= 299);
+
+			switch (Verb)
+			{
+			case RequestVerb::Head:
+				bSuccess = Is2xx || Is404;
+				VerbStr = TEXT("querying");
+				break;
+			case RequestVerb::Get:
+				bSuccess = Is2xx || Is404;
+				VerbStr = TEXT("fetching");
+				AdditionalInfo = FString::Printf(TEXT("Received: %d bytes."), BytesReceived);
+				break;
+			case RequestVerb::Put:
+				bSuccess = Is2xx;
+				VerbStr = TEXT("updating");
+				AdditionalInfo = FString::Printf(TEXT("Sent: %d bytes."), BytesSent);
+				break;
+			case RequestVerb::Post:
+				bSuccess = Is2xx || Is404;
+				VerbStr = TEXT("posting");
+				break;
+			case RequestVerb::Delete:
+				bSuccess = Is2xx || Is404;
+				VerbStr = TEXT("deleting");
+				break;
+			}
+
+			if (bSuccess)
+			{
+				UE_LOG(
+					LogZenHttp,
+					Verbose,
+					TEXT("Finished %s zen data (response %ld) from %s. %s"),
+					VerbStr,
+					ResponseCode,
+					Uri,
+					*AdditionalInfo
+				);
+			}
+			else if (bLogErrors)
+			{
+				// Print the response body if we got one, otherwise print header.
+				FString Response = GetAnsiBufferAsString(ResponseBuffer.Num() > 0 ? ResponseBuffer : ResponseHeader);
+				Response.ReplaceCharInline(TEXT('\n'), TEXT(' '));
+				Response.ReplaceCharInline(TEXT('\r'), TEXT(' '));
+				UE_LOG(
+					LogZenHttp,
+					Error,
+					TEXT("Failed %s zen data (response %ld) from %s. Response: %s"),
+					VerbStr,
+					ResponseCode,
+					Uri,
+					*Response
+				);
+			}
+		}
+		else if (bLogErrors)
+		{
+			UE_LOG(
+				LogZenHttp,
+				Error,
+				TEXT("Error while connecting to %s: %s"),
+				*Domain,
+				ANSI_TO_TCHAR(curl_easy_strerror(Result))
+			);
+		}
+	}
+
+	FString FZenHttpRequest::GetAnsiBufferAsString(const TArray64<uint8>& Buffer)
+	{
+		// Content is NOT null-terminated; we need to specify lengths here
+		FUTF8ToTCHAR TCHARData(reinterpret_cast<const ANSICHAR*>(Buffer.GetData()), IntCastChecked<int32>(Buffer.Num()));
+		return FString(TCHARData.Length(), TCHARData.Get());
+	}
+
+	FCbObjectView FZenHttpRequest::GetResponseAsObject() const
+	{
+		return FCbObjectView(ResponseBuffer.GetData());
+	}
+
+	size_t FZenHttpRequest::FStatics::StaticDebugCallback(CURL* Handle, curl_infotype DebugInfoType, char* DebugInfo, size_t DebugInfoSize, void* UserData)
+	{
+		FZenHttpRequest* Request = static_cast<FZenHttpRequest*>(UserData);
+
+		switch (DebugInfoType)
+		{
+			case CURLINFO_TEXT:
+				{
+					// Truncate at 1023 characters. This is just an arbitrary number based on a buffer size seen in
+					// the libcurl code.
+					DebugInfoSize = FMath::Min(DebugInfoSize, (size_t)1023);
+
+					// Calculate the actual length of the string due to incorrect use of snprintf() in lib/vtls/openssl.c.
+					char* FoundNulPtr = (char*)memchr(DebugInfo, 0, DebugInfoSize);
+					int CalculatedSize = FoundNulPtr != nullptr ? FoundNulPtr - DebugInfo : DebugInfoSize;
+
+					auto ConvertedString = StringCast<TCHAR>(static_cast<const ANSICHAR*>(DebugInfo), CalculatedSize);
+					FString DebugText(ConvertedString.Length(), ConvertedString.Get());
+					DebugText.ReplaceInline(TEXT("\n"), TEXT(""), ESearchCase::CaseSensitive);
+					DebugText.ReplaceInline(TEXT("\r"), TEXT(""), ESearchCase::CaseSensitive);
+					UE_LOG(LogZenHttp, VeryVerbose, TEXT("CURL %p: '%s'"), Request, *DebugText);
+				}
+				break;
+
+			case CURLINFO_HEADER_IN:
+				UE_LOG(LogZenHttp, VeryVerbose, TEXT("CURL %p: Received header (%zd bytes)"), Request, DebugInfoSize);
+				UE_LOG(LogZenHttp, VeryVerbose, TEXT("CURL HEADER <<< %*S"), DebugInfoSize, DebugInfo);
+				break;
+
+			case CURLINFO_HEADER_OUT:
+				UE_LOG(LogZenHttp, VeryVerbose, TEXT("CURL %p: Send header (%zd bytes)"), Request, DebugInfoSize);
+				UE_LOG(LogZenHttp, VeryVerbose, TEXT("CURL HEADER >>> %*S"), DebugInfoSize, DebugInfo);
+				break;
+
+			case CURLINFO_DATA_IN:
+				UE_LOG(LogZenHttp, VeryVerbose, TEXT("CURL %p: Received data (%zd bytes)"), Request, DebugInfoSize);
+				break;
+
+			case CURLINFO_DATA_OUT:
+				UE_LOG(LogZenHttp, VeryVerbose, TEXT("CURL %p: Sent data (%zd bytes)"), Request, DebugInfoSize);
+				break;
+
+			case CURLINFO_SSL_DATA_IN:
+				UE_LOG(LogZenHttp, VeryVerbose, TEXT("CURL %p: Received SSL data (%zd bytes)"), Request, DebugInfoSize);
+				break;
+
+			case CURLINFO_SSL_DATA_OUT:
+				UE_LOG(LogZenHttp, VeryVerbose, TEXT("CURL %p: Sent SSL data (%zd bytes)"), Request, DebugInfoSize);
+				break;
+		}
+
+		return 0;
+	}
+
+	size_t FZenHttpRequest::FStatics::StaticReadFn(void* Ptr, size_t SizeInBlocks, size_t BlockSizeInBytes, void* UserData)
+	{
+		FZenHttpRequest* Request = static_cast<FZenHttpRequest*>(UserData);
+		check(Request->ReadDataView);
+		const FCompositeBuffer& ReadDataView = *Request->ReadDataView;
+
+		const size_t Offset = Request->BytesSent;
+		const size_t ReadSize = FMath::Min((size_t)ReadDataView.GetSize() - Offset, SizeInBlocks * BlockSizeInBytes);
+		check(ReadDataView.GetSize() >= Offset + ReadSize);
+
+		Memcpy(Ptr, ReadDataView, Offset, ReadSize);
+		Request->BytesSent += ReadSize;
+		return ReadSize;
+
+		return 0;
+	}
+
+	size_t FZenHttpRequest::FStatics::StaticWriteHeaderFn(void* Ptr, size_t SizeInBlocks, size_t BlockSizeInBytes, void* UserData)
+	{
+		FZenHttpRequest* Request = static_cast<FZenHttpRequest*>(UserData);
+		const size_t WriteSize = SizeInBlocks * BlockSizeInBytes;
+		TArray64<uint8>* WriteHeaderBufferPtr = Request->WriteHeaderBufferPtr;
+		if (WriteHeaderBufferPtr && WriteSize > 0)
+		{
+			const size_t CurrentBufferLength = WriteHeaderBufferPtr->Num();
+			if (CurrentBufferLength > 0)
+			{
+				// Remove the previous zero termination
+				(*WriteHeaderBufferPtr)[CurrentBufferLength - 1] = ' ';
+			}
+
+			// Write the header
+			WriteHeaderBufferPtr->Append((const uint8*)Ptr, WriteSize + 1);
+			(*WriteHeaderBufferPtr)[WriteHeaderBufferPtr->Num() - 1] = 0; // Zero terminate string
+			return WriteSize;
+		}
+		return 0;
+	}
+
+	size_t FZenHttpRequest::FStatics::StaticWriteBodyFn(void* Ptr, size_t SizeInBlocks, size_t BlockSizeInBytes, void* UserData)
+	{
+		FZenHttpRequest* Request = static_cast<FZenHttpRequest*>(UserData);
+		const int64 WriteSize = IntCastChecked<int64>(SizeInBlocks * BlockSizeInBytes);
+		TArray64<uint8>* WriteDataBufferPtr = Request->WriteDataBufferPtr;
+
+		if (WriteDataBufferPtr && WriteSize > 0)
+		{
+			// If this is the first part of the body being received, try to reserve 
+			// memory if content length is defined in the header.
+			if (Request->BytesReceived == 0 && Request->WriteHeaderBufferPtr)
+			{
+				static const ANSICHAR* ContentLengthHeaderStr = "Content-Length: ";
+				const ANSICHAR* Header = (const ANSICHAR*)Request->WriteHeaderBufferPtr->GetData();
+
+				if (const ANSICHAR* ContentLengthHeader = FCStringAnsi::Strstr(Header, ContentLengthHeaderStr))
+				{
+					int64 ContentLength = FCStringAnsi::Atoi64(ContentLengthHeader + strlen(ContentLengthHeaderStr));
+					if (ContentLength > 0)
+					{
+						WriteDataBufferPtr->Reserve(ContentLength);
+					}
+				}
+			}
+
+			// Write to the target buffer
+			WriteDataBufferPtr->Append((const uint8*)Ptr, WriteSize);
+			Request->BytesReceived += WriteSize;
+			return WriteSize;
+		}
+
+		return 0;
+	}
+
+	size_t FZenHttpRequest::FStatics::StaticSeekFn(void* UserData, curl_off_t Offset, int Origin)
+	{
+		FZenHttpRequest* Request = static_cast<FZenHttpRequest*>(UserData);
+		size_t NewPosition = 0;
+		uint64 ReadDataSize = Request->ReadDataView ? Request->ReadDataView->GetSize() : 0;
+
+		switch (Origin)
+		{
+			case SEEK_SET: 
+				NewPosition = Offset; 
+				break;
+			case SEEK_CUR: 
+				NewPosition = Request->BytesSent + Offset; 
+				break;
+			case SEEK_END: 
+				NewPosition = ReadDataSize + Offset; 
+				break;
+		}
+
+		// Make sure we don't seek outside of the buffer
+		if (NewPosition < 0 || NewPosition >= ReadDataSize)
+		{
+			return CURL_SEEKFUNC_FAIL;
+		}
+
+		// Update the used offset
+		Request->BytesSent = NewPosition;
+		return CURL_SEEKFUNC_OK;
+	}
+
+#if PLATFORM_WINDOWS
+	int FZenHttpRequest::FStatics::StaticSockoptFn(void* UserData, curl_socket_t CurlFd, curlsocktype Purpose)
+	{
+		// On Windows, loopback connections can take advantage of a faster code path optionally with this flag.
+		// This must be used by both the client and server side, and is only effective in the absence of
+		// Windows Filtering Platform (WFP) callouts which can be installed by security software.
+		// https://docs.microsoft.com/en-us/windows/win32/winsock/sio-loopback-fast-path
+		int LoopbackOptionValue = 1;
+		DWORD OptionNumberOfBytesReturned = 0;
+		WSAIoctl(CurlFd, SIO_LOOPBACK_FAST_PATH, &LoopbackOptionValue, sizeof(LoopbackOptionValue), NULL, 0, &OptionNumberOfBytesReturned, 0, 0);
+		return CURL_SOCKOPT_OK;
+	}
+#endif //PLATFORM_WINDOWS
+
+	//////////////////////////////////////////////////////////////////////////
+
+	FZenHttpRequestPool::FZenHttpRequestPool(FStringView InServiceUrl, uint32 PoolEntryCount)
+	{
+		int32 LastSlash = 0;
+		while (InServiceUrl.FindLastChar(TEXT('/'), /* out */ LastSlash) && (LastSlash == (InServiceUrl.Len() - 1)))
+		{
+			InServiceUrl.LeftChopInline(1);
+		}
+
+		Pool.SetNum(PoolEntryCount);
+
+		for (uint8 i = 0; i < Pool.Num(); ++i)
+		{
+			Pool[i].IsAllocated = 0u;
+			Pool[i].Request = new FZenHttpRequest(InServiceUrl, true /* bLogErrors */);
+		}
+	}
+
+	FZenHttpRequestPool::~FZenHttpRequestPool()
+	{
+		for (uint8 i = 0; i < Pool.Num(); ++i)
+		{
+			// No requests should be in use by now.
+			check(Pool[i].IsAllocated.load(std::memory_order_relaxed) == 0u);
+
+			delete Pool[i].Request;
+		}
+	}
+
+	/** Block until a request is free
+	  *
+	  * Once a request has been returned it is owned by the caller and
+	  * it needs to release it to the pool when work has been completed
+	  *
+	  * @return Usable request instance.
+	  */
+	FZenHttpRequest* FZenHttpRequestPool::WaitForFreeRequest()
+	{
+		TRACE_CPUPROFILER_EVENT_SCOPE(ZenHttp_WaitForConnPool);
+
+		while (true)
+		{
+			for (uint8 i = 0; i < Pool.Num(); ++i)
+			{
+				if (!Pool[i].IsAllocated.load(std::memory_order_relaxed))
+				{
+					uint8 Expected = 0u;
+					if (Pool[i].IsAllocated.compare_exchange_strong(Expected, 1u))
+					{
+						return Pool[i].Request;
+					}
+				}
+			}
+
+			// This should use a better mechanism like condition variables / events
+			FPlatformProcess::Sleep(UE_ZENDDC_BACKEND_WAIT_INTERVAL);
+		}
+	}
+
+	/**
+	  * Release request to the pool.
+	  * @param Request Request that should be freed. Note that any buffer owned by the request can now be reset.
+	  */
+	void FZenHttpRequestPool::ReleaseRequestToPool(FZenHttpRequest* Request)
+	{
+		for (uint8 i = 0; i < Pool.Num(); ++i)
+		{
+			if (Pool[i].Request == Request)
+			{
+				Request->Reset();
+				uint8 Expected = 1u;
+				Pool[i].IsAllocated.compare_exchange_strong(Expected, 0u);
+				return;
+			}
+		}
+		check(false);
+	}
+}
+#endif // UE_WITH_ZEN
diff --git a/Engine/Source/Developer/Zen/Private/ZenServerInterface.cpp b/Engine/Source/Developer/Zen/Private/ZenServerInterface.cpp
new file mode 100644
--- /dev/null
+++ b/Engine/Source/Developer/Zen/Private/ZenServerInterface.cpp
@@ -0,0 +1,1686 @@
+// Copyright Epic Games, Inc. All Rights Reserved.
+#include "ZenServerInterface.h"
+
+
+#include "ZenBackendUtils.h"
+#include "ZenSerialization.h"
+#include "ZenServerHttp.h"
+
+#include "Analytics.h"
+#include "AnalyticsEventAttribute.h"
+#include "Async/Async.h"
+#include "Dom/JsonValue.h"
+#include "HAL/FileManager.h"
+#include "HAL/Platform.h"
+#include "HAL/PlatformMisc.h"
+#include "HAL/PlatformProcess.h"
+#include "HAL/PlatformTime.h"
+#include "Logging/LogScopedCategoryAndVerbosityOverride.h"
+#include "Misc/App.h"
+#include "Misc/ConfigCacheIni.h"
+#include "Misc/FileHelper.h"
+#include "Misc/MonitoredProcess.h"
+#include "Misc/Paths.h"
+#include "Misc/PathViews.h"
+#include "Misc/ScopedSlowTask.h"
+#include "Misc/ScopeExit.h"
+#include "Serialization/CompactBinary.h"
+#include "Serialization/CompactBinarySerialization.h"
+#include "Serialization/CompactBinaryValidation.h"
+#include "Serialization/JsonReader.h"
+#include "Serialization/JsonSerializer.h"
+#include "Serialization/JsonWriter.h"
+#include "SocketSubsystem.h"
+#include "String/LexFromString.h"
+#include "ZenVersion.h"
+
+#if PLATFORM_WINDOWS
+#	include "Windows/AllowWindowsPlatformTypes.h"
+#	include <shellapi.h>
+#	include <synchapi.h>
+#	include <Windows.h>
+#	include "Windows/HideWindowsPlatformTypes.h"
+#endif
+
+#if PLATFORM_UNIX || PLATFORM_MAC
+#	include <sys/file.h>
+#	include <sys/sem.h>
+#endif
+
+#define ALLOW_SETTINGS_OVERRIDE_FROM_COMMANDLINE			(UE_SERVER || !(UE_BUILD_SHIPPING))
+
+namespace UE::Zen
+{
+
+DEFINE_LOG_CATEGORY_STATIC(LogZenServiceInstance, Log, All);
+
+static FZenVersion
+GetZenVersion(const FString& UtilityPath, const FString& ServicePath, const FString& VersionCachePath)
+{
+	IFileManager& FileManager = IFileManager::Get();
+	FDateTime UtilityExecutableModificationTime = FileManager.GetTimeStamp(*UtilityPath);
+	FDateTime ServiceExecutableModificationTime = FileManager.GetTimeStamp(*ServicePath);
+
+	auto GetFallbackVersion = [&UtilityExecutableModificationTime, &ServiceExecutableModificationTime]()
+	{
+		FZenVersion FallbackVersion;
+		if (UtilityExecutableModificationTime > ServiceExecutableModificationTime)
+		{
+			FallbackVersion.Details = UtilityExecutableModificationTime.ToString();
+			return FallbackVersion;
+		}
+		FallbackVersion.Details = ServiceExecutableModificationTime.ToString();
+		return FallbackVersion;
+	};
+
+	FDateTime VersionCacheModificationTime = FileManager.GetTimeStamp(*VersionCachePath);
+	FString VersionFileContents;
+	FZenVersion ComparableVersion;
+	if ((VersionCacheModificationTime < UtilityExecutableModificationTime) ||
+		(VersionCacheModificationTime < ServiceExecutableModificationTime) ||
+		!FFileHelper::LoadFileToString(VersionFileContents, *VersionCachePath) ||
+		!ComparableVersion.TryParse(*VersionFileContents))
+	{
+		FString AbsoluteUtilityPath = FPaths::ConvertRelativePathToFull(UtilityPath);
+		FMonitoredProcess MonitoredUtilityProcess(AbsoluteUtilityPath, TEXT("version --detailed"), FPaths::GetPath(UtilityPath), true);
+
+		bool bLaunched = MonitoredUtilityProcess.Launch();
+		checkf(bLaunched, TEXT("Failed to launch zen utility to gather version data: '%s'."), *UtilityPath);
+		if (!bLaunched)
+		{
+			return GetFallbackVersion();
+		}
+
+		while (MonitoredUtilityProcess.Update())
+		{
+			FPlatformProcess::Sleep(0.1f);
+			if (MonitoredUtilityProcess.GetDuration().GetTotalSeconds() > 10)
+			{
+				MonitoredUtilityProcess.Cancel(true);
+				checkf(false, TEXT("Cancelled launch of zen utility for gathering version data: '%s'."), *UtilityPath);
+				return GetFallbackVersion();
+			}
+		}
+
+		if (MonitoredUtilityProcess.GetReturnCode() != 0)
+		{
+			checkf(false, TEXT("Unexpected return code after launch of zen utility for gathering version data: '%s' (%d)."), *UtilityPath, MonitoredUtilityProcess.GetReturnCode());
+			return GetFallbackVersion();
+		}
+
+		FString VersionOutputString = MonitoredUtilityProcess.GetFullOutputWithoutDelegate().TrimStartAndEnd();
+		if (!ComparableVersion.TryParse(*VersionOutputString))
+		{
+			checkf(false, TEXT("Invalid version information after launch of zen utility for gathering version data: '%s' (%s)."), *UtilityPath, *VersionOutputString);
+			return GetFallbackVersion();
+		}
+
+		FFileHelper::SaveStringToFile(ComparableVersion.ToString(), *VersionCachePath);
+	}
+
+	return ComparableVersion;
+}
+
+static bool
+IsInstallVersionOutOfDate(const FString& InTreeUtilityPath, const FString& InstallUtilityPath, const FString& InTreeServicePath, const FString& InstallServicePath, FString& OutInTreeVersionCache, FString& OutInstallVersionCache)
+{
+	OutInTreeVersionCache = FPaths::Combine(FPaths::EngineSavedDir(), TEXT("Zen") TEXT("zen.version"));
+	OutInstallVersionCache = FPaths::SetExtension(InstallUtilityPath, TEXT("version"));
+
+	IFileManager& FileManager = IFileManager::Get();
+	if (!FileManager.FileExists(*InstallUtilityPath) || !FileManager.FileExists(*InstallServicePath))
+	{
+		return true;
+	}
+
+	FZenVersion InTreeVersion = GetZenVersion(InTreeUtilityPath, InTreeServicePath, OutInTreeVersionCache);
+	FZenVersion InstallVersion = GetZenVersion(InstallUtilityPath, InstallServicePath, OutInstallVersionCache);
+
+	return InstallVersion < InTreeVersion;
+}
+
+static bool
+AttemptFileCopyWithRetries(const TCHAR* Dst, const TCHAR* Src, double RetryDurationSeconds)
+{
+	uint32 CopyResult = IFileManager::Get().Copy(Dst, Src, true, true, false);
+	uint64 CopyWaitStartTime = FPlatformTime::Cycles64();
+	while (CopyResult != COPY_OK)
+	{
+		double CopyWaitDuration = FPlatformTime::ToSeconds64(FPlatformTime::Cycles64() - CopyWaitStartTime);
+		if (CopyWaitDuration < RetryDurationSeconds)
+		{
+			FPlatformProcess::Sleep(0.01f);
+		}
+		else
+		{
+			break;
+		}
+		CopyResult = IFileManager::Get().Copy(Dst, Src, true, true, false);
+	}
+
+	return CopyResult == COPY_OK;
+}
+
+static void
+DetermineLocalDataCachePath(const TCHAR* ConfigSection, FString& DataPath)
+{
+	FString DataPathEnvOverride;
+	if (GConfig->GetString(ConfigSection, TEXT("LocalDataCachePathEnvOverride"), DataPathEnvOverride, GEngineIni))
+	{
+		FString DataPathEnvOverrideValue = FPlatformMisc::GetEnvironmentVariable(*DataPathEnvOverride);
+		if (!DataPathEnvOverrideValue.IsEmpty())
+		{
+			DataPath = DataPathEnvOverrideValue;
+			UE_LOG(LogZenServiceInstance, Log, TEXT("Found environment variable %s=%s"), *DataPathEnvOverride, *DataPathEnvOverrideValue);
+		}
+
+		if (FPlatformMisc::GetStoredValue(TEXT("Epic Games"), TEXT("GlobalDataCachePath"), *DataPathEnvOverride, DataPathEnvOverrideValue))
+		{
+			if (!DataPathEnvOverrideValue.IsEmpty())
+			{
+				DataPath = DataPathEnvOverrideValue;
+				UE_LOG(LogZenServiceInstance, Log, TEXT("Found registry key GlobalDataCachePath %s=%s"), *DataPathEnvOverride, *DataPath);
+			}
+		}
+	}
+
+	FString DataPathEditorOverrideSetting;
+	if (GConfig->GetString(ConfigSection, TEXT("LocalDataCachePathEditorOverrideSetting"), DataPathEditorOverrideSetting, GEngineIni))
+	{
+		FString Setting = GConfig->GetStr(TEXT("/Script/UnrealEd.EditorSettings"), *DataPathEditorOverrideSetting, GEditorSettingsIni);
+		if (!Setting.IsEmpty())
+		{
+			FString SettingPath;
+			if (FParse::Value(*Setting, TEXT("Path="), SettingPath))
+			{
+				SettingPath = SettingPath.TrimQuotes();
+				if (!SettingPath.IsEmpty())
+				{
+					DataPath = SettingPath;
+					UE_LOG(LogZenServiceInstance, Log, TEXT("Found editor setting /Script/UnrealEd.EditorSettings.Path=%s"), *DataPath);
+				}
+			}
+		}
+	}
+}
+
+static void
+DetermineDataPath(const TCHAR* ConfigSection, FString& DataPath)
+{
+	auto NormalizeDataPath = [](const FString& InDataPath)
+	{
+		FString FinalPath = FPaths::ConvertRelativePathToFull(InDataPath);
+		FPaths::NormalizeDirectoryName(FinalPath);
+		return FinalPath;
+	};
+
+	// Zen commandline
+	FString CommandLineOverrideValue;
+	if (FParse::Value(FCommandLine::Get(), TEXT("ZenDataPath="), CommandLineOverrideValue))
+	{
+		DataPath = NormalizeDataPath(CommandLineOverrideValue);
+		UE_LOG(LogZenServiceInstance, Log, TEXT("Found command line override ZenDataPath=%s"), *CommandLineOverrideValue);
+		return;
+	}
+
+	// Zen subprocess environment
+	FString SubprocessDataPathEnvOverrideValue;
+	SubprocessDataPathEnvOverrideValue = FPlatformMisc::GetEnvironmentVariable(TEXT("UE-ZenSubprocessDataPath"));
+	if (!SubprocessDataPathEnvOverrideValue.IsEmpty())
+	{
+		DataPath = NormalizeDataPath(SubprocessDataPathEnvOverrideValue);
+		UE_LOG(LogZenServiceInstance, Log, TEXT("Found subprocess environment variable UE-ZenSubprocessDataPath=%s"), *SubprocessDataPathEnvOverrideValue);
+		return;
+	}
+
+	// Zen registry/stored
+	FString DataPathEnvOverrideValue;
+	if (FPlatformMisc::GetStoredValue(TEXT("Epic Games"), TEXT("Zen"), TEXT("DataPath"), DataPathEnvOverrideValue))
+	{
+		if (!DataPathEnvOverrideValue.IsEmpty())
+		{
+			DataPath = NormalizeDataPath(DataPathEnvOverrideValue);
+			UE_LOG(LogZenServiceInstance, Log, TEXT("Found registry key Zen DataPath=%s"), *DataPath);
+			return;
+		}
+	}
+
+	// Zen environment
+	DataPathEnvOverrideValue = FPlatformMisc::GetEnvironmentVariable(TEXT("UE-ZenDataPath"));
+	if (!DataPathEnvOverrideValue.IsEmpty())
+	{
+		DataPath = NormalizeDataPath(DataPathEnvOverrideValue);
+		UE_LOG(LogZenServiceInstance, Log, TEXT("Found environment variable UE-ZenDataPath=%s"), *DataPathEnvOverrideValue);
+		return;
+	}
+
+	// Follow local DDC (if outside workspace)
+	FString LocalDataCachePath;
+	DetermineLocalDataCachePath(ConfigSection, LocalDataCachePath);
+	if (!LocalDataCachePath.IsEmpty() && (LocalDataCachePath != TEXT("None")) && !FPaths::IsUnderDirectory(LocalDataCachePath, FPaths::RootDir()))
+	{
+		DataPath = NormalizeDataPath(FPaths::Combine(LocalDataCachePath, TEXT("Zen")));
+		return;
+	}
+
+	// Zen config default
+	GConfig->GetString(ConfigSection, TEXT("DataPath"), DataPath, GEngineIni);
+	DataPath = NormalizeDataPath(DataPath);
+
+	check(!DataPath.IsEmpty())
+}
+
+static void
+ReadUInt16FromConfig(const TCHAR* Section, const TCHAR* Key, uint16& Value, const FString& ConfigFile)
+{
+	int32 ValueInt32 = Value;
+	GConfig->GetInt(Section, Key, ValueInt32, ConfigFile);
+	Value = (uint16)ValueInt32;
+}
+
+static bool
+IsLocalHost(const FString& Host)
+{
+	if (Host.Compare(FString(TEXT("localhost")), ESearchCase::IgnoreCase) == 0)
+	{
+		return true;
+	}
+
+	if (Host.Compare(FString(TEXT("127.0.0.1"))) == 0)
+	{
+		return true;
+	}
+
+	if (Host.Compare(FString(TEXT("[::1]"))) == 0)
+	{
+		return true;
+	}
+
+	ISocketSubsystem& SocketSubsystem = *ISocketSubsystem::Get();
+
+	const TSharedPtr<FInternetAddr> Addr = SocketSubsystem.GetAddressFromString(Host);
+	if (!Addr)
+	{
+		UE_LOG(LogZenServiceInstance, Warning, TEXT("Failed to get internet address from host '%s'"), *Host);
+		return false;
+	}
+
+	TArray<TSharedPtr<FInternetAddr>> LocalAddresses;
+	if (!SocketSubsystem.GetLocalAdapterAddresses(LocalAddresses))
+	{
+		UE_LOG(LogZenServiceInstance, Warning, TEXT("Failed to find local adapter addresses"));
+		return false;
+	}
+
+	for (const auto& Local : LocalAddresses)
+	{
+		if (*Local == *Addr)
+		{
+			return true;
+		}
+	}
+
+	return false;
+}
+
+void
+FServiceSettings::ReadFromConfig()
+{
+	check(GConfig && GConfig->IsReadyForUse());
+	const TCHAR* ConfigSection = TEXT("Zen");
+	bool bAutoLaunch = true;
+	GConfig->GetBool(ConfigSection, TEXT("AutoLaunch"), bAutoLaunch, GEngineIni);
+
+	if (bAutoLaunch)
+	{
+		if (!TryApplyAutoLaunchOverride())
+		{
+			// AutoLaunch settings
+			const TCHAR* AutoLaunchConfigSection = TEXT("Zen.AutoLaunch");
+			SettingsVariant.Emplace<FServiceAutoLaunchSettings>();
+			FServiceAutoLaunchSettings& AutoLaunchSettings = SettingsVariant.Get<FServiceAutoLaunchSettings>();
+
+			DetermineDataPath(AutoLaunchConfigSection, AutoLaunchSettings.DataPath);
+			GConfig->GetString(AutoLaunchConfigSection, TEXT("ExtraArgs"), AutoLaunchSettings.ExtraArgs, GEngineIni);
+
+			ReadUInt16FromConfig(AutoLaunchConfigSection, TEXT("DesiredPort"), AutoLaunchSettings.DesiredPort, GEngineIni);
+			GConfig->GetBool(AutoLaunchConfigSection, TEXT("ShowConsole"), AutoLaunchSettings.bShowConsole, GEngineIni);
+			GConfig->GetBool(AutoLaunchConfigSection, TEXT("LimitProcessLifetime"), AutoLaunchSettings.bLimitProcessLifetime, GEngineIni);
+
+			// Ensure that the zen data path is inherited by subprocesses
+			FPlatformMisc::SetEnvironmentVar(TEXT("UE-ZenSubprocessDataPath"), *AutoLaunchSettings.DataPath);
+		}
+	}
+	else
+	{
+		// ConnectExisting settings
+		const TCHAR* ConnectExistingConfigSection = TEXT("Zen.ConnectExisting");
+		SettingsVariant.Emplace<FServiceConnectSettings>();
+		FServiceConnectSettings& ConnectExistingSettings = SettingsVariant.Get<FServiceConnectSettings>();
+
+		GConfig->GetString(ConnectExistingConfigSection, TEXT("HostName"), ConnectExistingSettings.HostName, GEngineIni);
+		ReadUInt16FromConfig(ConnectExistingConfigSection, TEXT("Port"), ConnectExistingSettings.Port, GEngineIni);
+	}
+}
+
+void
+FServiceSettings::ReadFromJson(FJsonObject& JsonObject)
+{
+	if (TSharedPtr<FJsonValue> bAutoLaunchValue = JsonObject.Values.FindRef(TEXT("bAutoLaunch")))
+	{
+		if (bAutoLaunchValue->AsBool())
+		{
+			if (!TryApplyAutoLaunchOverride())
+			{
+				SettingsVariant.Emplace<FServiceAutoLaunchSettings>();
+				FServiceAutoLaunchSettings& AutoLaunchSettings = SettingsVariant.Get<FServiceAutoLaunchSettings>();
+
+				TSharedPtr<FJsonValue> AutoLaunchSettingsValue = JsonObject.Values.FindRef(TEXT("AutoLaunchSettings"));
+				if (AutoLaunchSettingsValue)
+				{
+					TSharedPtr<FJsonObject> AutoLaunchSettingsObject = AutoLaunchSettingsValue->AsObject();
+					AutoLaunchSettings.DataPath = AutoLaunchSettingsObject->Values.FindRef(TEXT("DataPath"))->AsString();
+					AutoLaunchSettings.ExtraArgs = AutoLaunchSettingsObject->Values.FindRef(TEXT("ExtraArgs"))->AsString();
+					AutoLaunchSettingsObject->Values.FindRef(TEXT("DesiredPort"))->TryGetNumber(AutoLaunchSettings.DesiredPort);
+					AutoLaunchSettingsObject->Values.FindRef(TEXT("ShowConsole"))->TryGetBool(AutoLaunchSettings.bShowConsole);
+					AutoLaunchSettingsObject->Values.FindRef(TEXT("LimitProcessLifetime"))->TryGetBool(AutoLaunchSettings.bLimitProcessLifetime);
+				}
+			}
+		}
+		else
+		{
+			SettingsVariant.Emplace<FServiceConnectSettings>();
+			FServiceConnectSettings& ConnectExistingSettings = SettingsVariant.Get<FServiceConnectSettings>();
+
+			TSharedPtr<FJsonValue> ConnectExistingSettingsValue = JsonObject.Values.FindRef(TEXT("ConnectExistingSettings"));
+			if (ConnectExistingSettingsValue)
+			{
+				TSharedPtr<FJsonObject> ConnectExistingSettingsObject = ConnectExistingSettingsValue->AsObject();
+				ConnectExistingSettings.HostName = ConnectExistingSettingsObject->Values.FindRef(TEXT("HostName"))->AsString();
+				ConnectExistingSettingsObject->Values.FindRef(TEXT("Port"))->TryGetNumber(ConnectExistingSettings.Port);
+			}
+		}
+
+	}
+}
+
+void
+FServiceSettings::ReadFromURL(FStringView InstanceURL)
+{
+	SettingsVariant.Emplace<FServiceConnectSettings>();
+	FServiceConnectSettings& ConnectExistingSettings = SettingsVariant.Get<FServiceConnectSettings>();
+
+	if (InstanceURL.StartsWith(TEXT("http://")))
+	{
+		InstanceURL.RightChopInline(7);
+	}
+
+	int32 PortDelimIndex = INDEX_NONE;
+	InstanceURL.FindLastChar(TEXT(':'), PortDelimIndex);
+	if (PortDelimIndex != INDEX_NONE)
+	{
+		ConnectExistingSettings.HostName = InstanceURL.Left(PortDelimIndex);
+		LexFromString(ConnectExistingSettings.Port, InstanceURL.RightChop(PortDelimIndex + 1));
+	}
+	else
+	{
+		ConnectExistingSettings.HostName = InstanceURL;
+		ConnectExistingSettings.Port = 1337;
+	}
+}
+
+void
+FServiceSettings::WriteToJson(TJsonWriter<TCHAR, TPrettyJsonPrintPolicy<TCHAR>>& Writer) const
+{
+	bool bAutoLaunch = IsAutoLaunch();
+	Writer.WriteValue(TEXT("bAutoLaunch"), bAutoLaunch);
+	if (bAutoLaunch)
+	{
+		const FServiceAutoLaunchSettings& AutoLaunchSettings = SettingsVariant.Get<FServiceAutoLaunchSettings>();
+		Writer.WriteObjectStart(TEXT("AutoLaunchSettings"));
+		Writer.WriteValue(TEXT("DataPath"), AutoLaunchSettings.DataPath);
+		Writer.WriteValue(TEXT("ExtraArgs"), AutoLaunchSettings.ExtraArgs);
+		Writer.WriteValue(TEXT("DesiredPort"), AutoLaunchSettings.DesiredPort);
+		Writer.WriteValue(TEXT("ShowConsole"), AutoLaunchSettings.bShowConsole);
+		Writer.WriteValue(TEXT("LimitProcessLifetime"), AutoLaunchSettings.bLimitProcessLifetime);
+		Writer.WriteObjectEnd();
+	}
+	else
+	{
+		const FServiceConnectSettings& ConnectExistingSettings = SettingsVariant.Get<FServiceConnectSettings>();
+		Writer.WriteObjectStart(TEXT("ConnectExistingSettings"));
+		Writer.WriteValue(TEXT("HostName"), ConnectExistingSettings.HostName);
+		Writer.WriteValue(TEXT("Port"), ConnectExistingSettings.Port);
+		Writer.WriteObjectEnd();
+	}
+}
+
+bool
+FServiceSettings::TryApplyAutoLaunchOverride()
+{
+#if ALLOW_SETTINGS_OVERRIDE_FROM_COMMANDLINE
+	if (FParse::Param(FCommandLine::Get(), TEXT("NoZenAutoLaunch")))
+	{
+		SettingsVariant.Emplace<FServiceConnectSettings>();
+		FServiceConnectSettings& ConnectExistingSettings = SettingsVariant.Get<FServiceConnectSettings>();
+		ConnectExistingSettings.HostName = TEXT("[::1]");
+		ConnectExistingSettings.Port = 1337;
+		return true;
+	}
+
+	FString Host;
+	if  (FParse::Value(FCommandLine::Get(), TEXT("-NoZenAutoLaunch="), Host))
+	{
+		SettingsVariant.Emplace<FServiceConnectSettings>();
+		FServiceConnectSettings& ConnectExistingSettings = SettingsVariant.Get<FServiceConnectSettings>();
+
+		int32 PortDelimIndex = INDEX_NONE;
+		if (Host.FindChar(TEXT(':'), PortDelimIndex))
+		{
+			ConnectExistingSettings.HostName = Host.Left(PortDelimIndex);
+			LexFromString(ConnectExistingSettings.Port, Host.RightChop(PortDelimIndex + 1));
+		}
+		else
+		{
+			ConnectExistingSettings.HostName = Host;
+			ConnectExistingSettings.Port = 1337;
+		}
+
+		return true;
+	}
+#endif
+	return false;
+}
+
+#if UE_WITH_ZEN
+
+uint16 FZenServiceInstance::AutoLaunchedPort = 0;
+
+static bool
+ReadCbLockFile(FStringView FileName, FCbObject& OutLockObject)
+{
+#if PLATFORM_WINDOWS
+	// Windows specific lock reading path
+	// Uses share flags that are unique to windows to allow us to read file contents while the file may be open for write AND delete by another process (zenserver).
+
+	uint32 Access = GENERIC_READ;
+	uint32 WinFlags = FILE_SHARE_READ | FILE_SHARE_WRITE | FILE_SHARE_DELETE;
+	uint32 Create = OPEN_EXISTING;
+
+	TStringBuilder<MAX_PATH> FullFileNameBuilder;
+	FPathViews::ToAbsolutePath(FileName, FullFileNameBuilder);
+	for (TCHAR& Char : MakeArrayView(FullFileNameBuilder))
+	{
+		if (Char == TEXT('/'))
+		{
+			Char = TEXT('\\');
+		}
+	}
+	if (FullFileNameBuilder.Len() >= MAX_PATH)
+	{
+		FullFileNameBuilder.Prepend(TEXTVIEW("\\\\?\\"));
+	}
+	HANDLE Handle = CreateFileW(FullFileNameBuilder.ToString(), Access, WinFlags, NULL, Create, FILE_ATTRIBUTE_NORMAL, NULL);
+	if (Handle != INVALID_HANDLE_VALUE)
+	{
+		ON_SCOPE_EXIT { CloseHandle(Handle); };
+		LARGE_INTEGER LI;
+		if (GetFileSizeEx(Handle, &LI))
+		{
+			checkf(LI.QuadPart == LI.u.LowPart, TEXT("Lock file exceeds supported 2GB limit."));
+			int32 FileSize32 = LI.u.LowPart;
+			FUniqueBuffer FileBytes = FUniqueBuffer::Alloc(FileSize32);
+			DWORD ReadBytes = 0;
+			if (ReadFile(Handle, FileBytes.GetData(), FileSize32, &ReadBytes, NULL) && (ReadBytes == FileSize32))
+			{
+				if (ValidateCompactBinary(FileBytes, ECbValidateMode::Default) == ECbValidateError::None)
+				{
+					OutLockObject = FCbObject(FileBytes.MoveToShared());
+					return true;
+				}
+			}
+		}
+	}
+	return false;
+#elif PLATFORM_UNIX || PLATFORM_MAC
+	TAnsiStringBuilder<256> LockFilePath;
+	LockFilePath << FileName;
+	int32 Fd = open(LockFilePath.ToString(), O_RDONLY);
+	if (Fd < 0)
+	{
+		return false;
+	}
+
+	// If we can claim the lock then it's an orphaned lock file and should be
+	// ignored. Not ideal as there's a period of time when the lock can be
+	// held unncessarily.
+	int32 LockRet = flock(Fd, LOCK_EX | LOCK_NB);
+	if (LockRet >= 0)
+	{
+		unlink(LockFilePath.ToString());
+		flock(Fd, LOCK_UN);
+		close(Fd);
+		return false;
+	}
+
+	if (errno != EWOULDBLOCK && errno != EAGAIN)
+	{
+		return false;
+	}
+
+	struct stat Stat;
+	fstat(Fd, &Stat);
+	uint64 FileSize = uint64(Stat.st_size);
+
+	bool bSuccess = false;
+	FUniqueBuffer FileBytes = FUniqueBuffer::Alloc(FileSize);
+	if (read(Fd, FileBytes.GetData(), FileSize) == FileSize)
+	{
+		if (ValidateCompactBinary(FileBytes, ECbValidateMode::Default) == ECbValidateError::None)
+		{
+			OutLockObject = FCbObject(FileBytes.MoveToShared());
+			bSuccess = true;
+		}
+	}
+
+	close(Fd);
+	return bSuccess;
+#endif
+}
+
+static bool
+IsLockFileLocked(const TCHAR* FileName, bool bAttemptCleanUp=false)
+{
+#if PLATFORM_WINDOWS
+	if (bAttemptCleanUp)
+	{
+		IFileManager::Get().Delete(FileName, false, false, true);
+	}
+	return IFileManager::Get().FileExists(FileName);
+#elif PLATFORM_UNIX || PLATFORM_MAC
+	TAnsiStringBuilder<256> LockFilePath;
+	LockFilePath << FileName;
+	int32 Fd = open(LockFilePath.ToString(), O_RDONLY);
+	if (Fd < 0)
+	{
+		return false;
+	}
+
+	int32 LockRet = flock(Fd, LOCK_EX | LOCK_NB);
+	if (LockRet < 0)
+	{
+		close(Fd);
+		return errno == EWOULDBLOCK || errno == EAGAIN;
+	}
+
+	// Consider the lock file as orphaned if we we managed to claim the lock for
+	// it. Might as well delete it while we own it.
+	unlink(LockFilePath.ToString());
+
+	flock(Fd, LOCK_UN);
+	close(Fd);
+
+	return true;
+#endif
+}
+
+static void
+RequestZenShutdownOnPort(uint16 Port)
+{
+#if PLATFORM_WINDOWS
+	HANDLE Handle = OpenEventW(EVENT_MODIFY_STATE, false, *WriteToWideString<64>(WIDETEXT("Zen_"), Port, WIDETEXT("_Shutdown")));
+	if (Handle != INVALID_HANDLE_VALUE)
+	{
+		ON_SCOPE_EXIT{ CloseHandle(Handle); };
+		SetEvent(Handle);
+	}
+#elif PLATFORM_UNIX || PLATFORM_MAC
+	TAnsiStringBuilder<64> EventPath;
+	EventPath << "/tmp/Zen_" << Port << "_Shutdown";
+
+	key_t IpcKey = ftok(EventPath.ToString(), 1);
+	if (IpcKey < 0)
+	{
+		return;
+	}
+
+	int Semaphore = semget(IpcKey, 1, 0600);
+	if (Semaphore < 0)
+	{
+		return;
+	}
+
+	semctl(Semaphore, 0, SETVAL, 0);
+	semctl(Semaphore, 0, IPC_RMID);
+#else
+	static_assert(false, "Missing implementation for Zen named shutdown events");
+#endif
+}
+
+static bool
+WaitForZenShutdown(const TCHAR* LockFilePath, double MaximumWaitDurationSeconds)
+{
+	uint64 ZenShutdownWaitStartTime = FPlatformTime::Cycles64();
+	while (IsLockFileLocked(LockFilePath))
+	{
+		double ZenShutdownWaitDuration = FPlatformTime::ToSeconds64(FPlatformTime::Cycles64() - ZenShutdownWaitStartTime);
+		if (ZenShutdownWaitDuration < MaximumWaitDurationSeconds)
+		{
+			FPlatformProcess::Sleep(0.01f);
+		}
+		else
+		{
+			return false;
+		}
+	}
+	return true;
+}
+
+static bool IsProcessActive(const TCHAR* ExecutablePath)
+{
+	FString NormalizedExecutablePath(ExecutablePath);
+	FPaths::NormalizeFilename(NormalizedExecutablePath);
+	FPlatformProcess::FProcEnumerator ProcIter;
+	while (ProcIter.MoveNext())
+	{
+		FPlatformProcess::FProcEnumInfo ProcInfo = ProcIter.GetCurrent();
+		FString Candidate = ProcInfo.GetFullPath();
+		FPaths::NormalizeFilename(Candidate);
+		if (Candidate == NormalizedExecutablePath)
+		{
+			return true;
+		}
+	}
+	return false;
+}
+
+static FString
+DetermineCmdLineWithoutTransientComponents(const FServiceAutoLaunchSettings& InSettings, int16 OverrideDesiredPort)
+{
+	FString PlatformDataPath(InSettings.DataPath);
+	FPaths::MakePlatformFilename(PlatformDataPath);
+
+	FString Parms;
+	Parms.Appendf(TEXT("--port %d --data-dir \"%s\""),
+		OverrideDesiredPort,
+		*PlatformDataPath);
+
+	if (!InSettings.ExtraArgs.IsEmpty())
+	{
+		Parms.AppendChar(TEXT(' '));
+		Parms.Append(InSettings.ExtraArgs);
+	}
+
+	FString LogCommandLineOverrideValue;
+	if (FParse::Value(FCommandLine::Get(), TEXT("ZenLogPath="), LogCommandLineOverrideValue))
+	{
+		if (!LogCommandLineOverrideValue.IsEmpty())
+		{
+			Parms.Appendf(TEXT(" --abslog \"%s\""),
+				*FPaths::ConvertRelativePathToFull(LogCommandLineOverrideValue));
+		}
+	}
+
+	FString CfgCommandLineOverrideValue;
+	if (FParse::Value(FCommandLine::Get(), TEXT("ZenCfgPath="), CfgCommandLineOverrideValue))
+	{
+		if (!CfgCommandLineOverrideValue.IsEmpty())
+		{
+			Parms.Appendf(TEXT(" --config \"%s\""),
+				*FPaths::ConvertRelativePathToFull(CfgCommandLineOverrideValue));
+		}
+	}
+
+	return Parms;
+}
+
+bool
+TryGetLocalServiceRunContext(FZenLocalServiceRunContext& OutContext)
+{
+	FString LocalServiceRunContextPath = FPaths::SetExtension(GetLocalServiceInstallPath(), TEXT(".runcontext"));
+	return OutContext.ReadFromJsonFile(*LocalServiceRunContextPath);
+}
+
+bool
+FZenLocalServiceRunContext::ReadFromJson(FJsonObject& JsonObject)
+{
+	Executable = JsonObject.Values.FindRef(TEXT("Executable"))->AsString();
+	CommandlineArguments = JsonObject.Values.FindRef(TEXT("CommandlineArguments"))->AsString();
+	WorkingDirectory = JsonObject.Values.FindRef(TEXT("WorkingDirectory"))->AsString();
+	DataPath = JsonObject.Values.FindRef(TEXT("DataPath"))->AsString();
+	bShowConsole = JsonObject.Values.FindRef(TEXT("ShowConsole"))->AsBool();
+	return true;
+}
+
+void
+FZenLocalServiceRunContext::WriteToJson(TJsonWriter<TCHAR, TPrettyJsonPrintPolicy<TCHAR>>& Writer) const
+{
+	Writer.WriteValue(TEXT("Executable"), Executable);
+	Writer.WriteValue(TEXT("CommandlineArguments"), CommandlineArguments);
+	Writer.WriteValue(TEXT("WorkingDirectory"), WorkingDirectory);
+	Writer.WriteValue(TEXT("DataPath"), DataPath);
+	Writer.WriteValue(TEXT("ShowConsole"), bShowConsole);
+}
+
+bool
+FZenLocalServiceRunContext::ReadFromJsonFile(const TCHAR* Filename)
+{
+	FString JsonText;
+	if (!FFileHelper::LoadFileToString(JsonText, Filename))
+	{
+		return false;
+	}
+
+	TSharedPtr<FJsonObject> JsonObject;
+	TSharedRef<TJsonReader<TCHAR>> Reader = TJsonReaderFactory<TCHAR>::Create(JsonText);
+	if (!FJsonSerializer::Deserialize(Reader, JsonObject) || !JsonObject.IsValid())
+	{
+		return false;
+	}
+
+	return ReadFromJson(*JsonObject);
+}
+
+bool
+FZenLocalServiceRunContext::WriteToJsonFile(const TCHAR* Filename) const
+{
+	FString JsonTcharText;
+	TSharedRef<TJsonWriter<TCHAR, TPrettyJsonPrintPolicy<TCHAR>>> Writer = TJsonWriterFactory<TCHAR, TPrettyJsonPrintPolicy<TCHAR>>::Create(&JsonTcharText);
+	Writer->WriteObjectStart();
+	WriteToJson(*Writer);
+	Writer->WriteObjectEnd();
+	Writer->Close();
+
+	if (!FFileHelper::SaveStringToFile(JsonTcharText, Filename))
+	{
+		return false;
+	}
+
+	return true;
+}
+
+bool
+IsLocalServiceRunning(const TCHAR* DataPath, uint16* OutPort)
+{
+	// TODO: Instead of using the lock file, this could use the shared memory system state named "Global\ZenMap" (see zenserverprocess.{h,cpp} in zen codebase)
+	const FString LockFilePath = FPaths::Combine(DataPath, TEXT(".lock"));
+	if (IsLockFileLocked(*LockFilePath, true))
+	{
+		if (OutPort)
+		{
+			// If an instance is running with this data path, check if we can use it and what port it is on
+			*OutPort = 0;
+			FCbObject LockObject;
+			if (ReadCbLockFile(LockFilePath, LockObject))
+			{
+				bool bIsReady = LockObject["ready"].AsBool();
+				if (bIsReady)
+				{
+					*OutPort = LockObject["port"].AsUInt16();
+				}
+			}
+		}
+
+		return true;
+	}
+	return false;
+}
+
+FProcHandle
+StartLocalService(const FZenLocalServiceRunContext& Context, const TCHAR* TransientArgs)
+{
+	FString Parms = Context.GetCommandlineArguments();
+	if (TransientArgs)
+	{
+		Parms.Appendf(TEXT(" %s"), TransientArgs);
+	}
+
+	FProcHandle Proc;
+#if PLATFORM_WINDOWS
+	FString PlatformExecutable = Context.GetExecutable();
+	FPaths::MakePlatformFilename(PlatformExecutable);
+	FString PlatformWorkingDirectory = Context.GetWorkingDirectory();
+	FPaths::MakePlatformFilename(PlatformWorkingDirectory);
+	{
+		// Attempt non-elevated launch
+		STARTUPINFO StartupInfo = {
+			sizeof(STARTUPINFO),
+			NULL, NULL, NULL,
+			(::DWORD)CW_USEDEFAULT,
+			(::DWORD)CW_USEDEFAULT,
+			(::DWORD)CW_USEDEFAULT,
+			(::DWORD)CW_USEDEFAULT,
+			(::DWORD)0, (::DWORD)0, (::DWORD)0,
+			(::DWORD)STARTF_USESHOWWINDOW,
+			(::WORD)(Context.GetShowConsole() ? SW_SHOWMINNOACTIVE : SW_HIDE),
+			0, NULL,
+			HANDLE(nullptr),
+			HANDLE(nullptr),
+			HANDLE(nullptr)
+		};
+
+		FString CommandLine = FString::Printf(TEXT("\"%s\" %s"), *PlatformExecutable, *Parms);
+		PROCESS_INFORMATION ProcInfo;
+		if (CreateProcess(NULL, CommandLine.GetCharArray().GetData(), nullptr, nullptr, false, (::DWORD)(NORMAL_PRIORITY_CLASS | DETACHED_PROCESS), nullptr, PlatformWorkingDirectory.GetCharArray().GetData(), &StartupInfo, &ProcInfo))
+		{
+			::CloseHandle(ProcInfo.hThread);
+			Proc = FProcHandle(ProcInfo.hProcess);
+		}
+
+	}
+	if (!Proc.IsValid())
+	{
+		// Fall back to elevated launch
+		SHELLEXECUTEINFO ShellExecuteInfo;
+		ZeroMemory(&ShellExecuteInfo, sizeof(ShellExecuteInfo));
+		ShellExecuteInfo.cbSize = sizeof(ShellExecuteInfo);
+		ShellExecuteInfo.fMask = SEE_MASK_UNICODE | SEE_MASK_NOCLOSEPROCESS;
+		ShellExecuteInfo.lpFile = *PlatformExecutable;
+		ShellExecuteInfo.lpDirectory = *PlatformWorkingDirectory;
+		ShellExecuteInfo.lpVerb = TEXT("runas");
+		ShellExecuteInfo.nShow = Context.GetShowConsole() ? SW_SHOWMINNOACTIVE : SW_HIDE;
+		ShellExecuteInfo.lpParameters = *Parms;
+
+		if (ShellExecuteEx(&ShellExecuteInfo))
+		{
+			Proc = FProcHandle(ShellExecuteInfo.hProcess);
+		}
+	}
+#else
+	{
+		bool bLaunchDetached = true;
+		bool bLaunchHidden = true;
+		bool bLaunchReallyHidden = !Context.GetShowConsole();
+		uint32* OutProcessID = nullptr;
+		int32 PriorityModifier = 0;
+		void* PipeWriteChild = nullptr;
+		void* PipeReadChild = nullptr;
+		Proc = FPlatformProcess::CreateProc(
+			*Context.GetExecutable(),
+			*Parms,
+			bLaunchDetached,
+			bLaunchHidden,
+			bLaunchReallyHidden,
+			OutProcessID,
+			PriorityModifier,
+			*Context.GetWorkingDirectory(),
+			PipeWriteChild,
+			PipeReadChild);
+	}
+#endif
+	return Proc;
+}
+
+bool
+StopLocalService(const TCHAR* DataPath, double MaximumWaitDurationSeconds)
+{
+	uint16 CurrentPort = 0;
+	if (IsLocalServiceRunning(DataPath, &CurrentPort))
+	{
+		if (CurrentPort == 0)
+		{
+			return false;
+		}
+
+		RequestZenShutdownOnPort(CurrentPort);
+		return WaitForZenShutdown(*FPaths::Combine(DataPath, TEXT(".lock")), MaximumWaitDurationSeconds);
+	}
+	return true;
+}
+
+FString
+GetLocalServiceInstallPath()
+{
+	return FPaths::ConvertRelativePathToFull(FPaths::Combine(FPlatformProcess::ApplicationSettingsDir(), TEXT("Zen\\Install"),
+#if PLATFORM_WINDOWS
+		TEXT("zenserver.exe")
+#else
+		TEXT("zenserver")
+#endif
+		));
+}
+
+FString
+GetLocalServiceInstallVersion(bool bDetailed)
+{
+	FString InstallUtilityPath = GetLocalInstallUtilityPath();
+	FString InstallVersionCache = FPaths::SetExtension(InstallUtilityPath, TEXT("version"));
+
+	IFileManager& FileManager = IFileManager::Get();
+	if (!FileManager.FileExists(*InstallUtilityPath))
+	{
+		return FZenVersion().ToString(bDetailed);
+	}
+
+	FZenVersion InstallVersion = GetZenVersion(InstallUtilityPath, GetLocalServiceInstallPath(), InstallVersionCache);
+
+	return InstallVersion.ToString(bDetailed);
+}
+
+FString
+GetLocalInstallUtilityPath()
+{
+	return FPaths::ConvertRelativePathToFull(FPaths::Combine(FPlatformProcess::ApplicationSettingsDir(), TEXT("Zen\\Install"),
+#if PLATFORM_WINDOWS
+		TEXT("zen.exe")
+#else
+		TEXT("zen")
+#endif
+		));
+}
+
+static bool GIsDefaultServicePresent = false;
+
+FZenServiceInstance& GetDefaultServiceInstance()
+{
+	static FZenServiceInstance DefaultServiceInstance;
+	GIsDefaultServicePresent = true;
+	return DefaultServiceInstance;
+}
+
+bool IsDefaultServicePresent()
+{
+	return GIsDefaultServicePresent;
+}
+
+FScopeZenService::FScopeZenService()
+	: FScopeZenService(FStringView())
+{
+}
+
+FScopeZenService::FScopeZenService(FStringView InstanceURL)
+{
+	if (!InstanceURL.IsEmpty() && !InstanceURL.Equals(TEXT("<DefaultInstance>")))
+	{
+		UniqueNonDefaultInstance = MakeUnique<FZenServiceInstance>(InstanceURL);
+		ServiceInstance = UniqueNonDefaultInstance.Get();
+	}
+	else
+	{
+		ServiceInstance = &GetDefaultServiceInstance();
+	}
+}
+
+FScopeZenService::FScopeZenService(FServiceSettings&& InSettings)
+{
+	UniqueNonDefaultInstance = MakeUnique<FZenServiceInstance>(MoveTemp(InSettings));
+	ServiceInstance = UniqueNonDefaultInstance.Get();
+}
+
+FScopeZenService::~FScopeZenService()
+{}
+
+FZenServiceInstance::FZenServiceInstance()
+: FZenServiceInstance(FStringView())
+{
+}
+
+FZenServiceInstance::FZenServiceInstance(FStringView InstanceURL)
+{
+	if (InstanceURL.IsEmpty())
+	{
+		Settings.ReadFromConfig();
+	}
+	else
+	{
+		Settings.ReadFromURL(InstanceURL);
+	}
+
+	Initialize();
+}
+
+FZenServiceInstance::FZenServiceInstance(FServiceSettings&& InSettings)
+: Settings(MoveTemp(InSettings))
+{
+	Initialize();
+}
+
+FZenServiceInstance::~FZenServiceInstance()
+{
+	if (FAnalytics::IsAvailable())
+	{
+		FAnalytics::Get().GetEventCallback(TEXT("Core.Loading"))->RemoveAll(this);
+	}
+}
+
+bool 
+FZenServiceInstance::IsServiceRunning()
+{
+	return !Settings.IsAutoLaunch() || bHasLaunchedLocal;
+}
+
+bool 
+FZenServiceInstance::IsServiceReady()
+{
+	if (IsServiceRunning())
+	{
+		TStringBuilder<128> ZenDomain;
+		ZenDomain << HostName << TEXT(":") << Port;
+		Zen::FZenHttpRequest Request(ZenDomain.ToString(), false);
+		Zen::FZenHttpRequest::Result Result = Request.PerformBlockingDownload(TEXTVIEW("health/ready"), nullptr, Zen::EContentType::Text);
+		
+		if (Result == Zen::FZenHttpRequest::Result::Success && Zen::IsSuccessCode(Request.GetResponseCode()))
+		{
+			UE_LOG(LogZenServiceInstance, Display, TEXT("Z$ HTTP DDC service status: %s."), *Request.GetResponseAsString());
+			return true;
+		}
+		else
+		{
+			UE_LOG(LogZenServiceInstance, Warning, TEXT("Unable to reach Z$ HTTP DDC service at %s. Status: %d . Response: %s"), ZenDomain.ToString(), Request.GetResponseCode(), *Request.GetResponseAsString());
+		}
+	}
+	return false;
+}
+
+uint16
+FZenServiceInstance::GetAutoLaunchedPort()
+{
+	return AutoLaunchedPort;
+}
+
+void
+FZenServiceInstance::Initialize()
+{
+	if (Settings.IsAutoLaunch())
+	{
+		bHasLaunchedLocal = AutoLaunch(Settings.SettingsVariant.Get<FServiceAutoLaunchSettings>(), ConditionalUpdateLocalInstall(), HostName, Port);
+		if (bHasLaunchedLocal)
+		{
+			AutoLaunchedPort = Port;
+			bIsRunningLocally = true;
+		}
+	}
+	else
+	{
+		const FServiceConnectSettings& ConnectExistingSettings = Settings.SettingsVariant.Get<FServiceConnectSettings>();
+		HostName = ConnectExistingSettings.HostName;
+		Port = ConnectExistingSettings.Port;
+		bIsRunningLocally = IsLocalHost(HostName);
+	}
+	URL = WriteToString<64>(TEXT("http://"), HostName, TEXT(":"), Port, TEXT("/"));
+
+	FAnalytics::Get().GetEventCallback(TEXT("Core.Loading"))->AddRaw(this,&FZenServiceInstance::OnAnalyticsEvent);
+}
+
+void
+FZenServiceInstance::PromptUserToStopRunningServerInstance(const FString& ServerFilePath)
+{
+	if (FApp::IsUnattended())
+	{
+		// Do not ask if there is no one to show a message
+		return;
+	}
+
+	FText ZenUpdatePromptTitle = NSLOCTEXT("Zen", "Zen_UpdatePromptTitle", "Update required");
+	FText ZenUpdatePromptText = FText::Format(NSLOCTEXT("Zen", "Zen_UpdatePromptText", "ZenServer needs to be updated to a new version. Please shut down Unreal Editor and any tools that are using the ZenServer at '{0}'"), FText::FromString(ServerFilePath));
+	FPlatformMisc::MessageBoxExt(EAppMsgType::Ok, *ZenUpdatePromptText.ToString(), *ZenUpdatePromptTitle.ToString());
+}
+
+FString
+FZenServiceInstance::ConditionalUpdateLocalInstall()
+{
+	FString InTreeUtilityPath = FPaths::ConvertRelativePathToFull(FPlatformProcess::GenerateApplicationPath(TEXT("zen"), EBuildConfiguration::Development));
+	FString InstallUtilityPath = GetLocalInstallUtilityPath();
+
+	FString InTreeServicePath = FPaths::ConvertRelativePathToFull(FPlatformProcess::GenerateApplicationPath(TEXT("zenserver"), EBuildConfiguration::Development));
+	FString InstallServicePath = GetLocalServiceInstallPath();
+
+	IFileManager& FileManager = IFileManager::Get();
+
+	bool bMainExecutablesUpdated = false;
+	FString InTreeVersionCache, InstallVersionCache;
+	if (IsInstallVersionOutOfDate(InTreeUtilityPath, InstallUtilityPath, InTreeServicePath, InstallServicePath, InTreeVersionCache, InstallVersionCache))
+	{
+		if (IsProcessActive(*InstallServicePath))
+		{
+			FString DataPath = Settings.SettingsVariant.Get<FServiceAutoLaunchSettings>().DataPath;
+			StopLocalService(*DataPath);
+
+			if (IsLocalServiceRunning(*DataPath))
+			{
+				PromptUserToStopRunningServerInstance(InstallServicePath);
+			}
+		}
+
+		// Even after waiting for the lock file to be removed, the executable may have a period where it can't be overwritten as the process shuts down
+		// so any attempt to overwrite it should have some tolerance for retrying.
+		bool bServiceExecutableUpdated = AttemptFileCopyWithRetries(*InstallServicePath, *InTreeServicePath, 5.0);
+		checkf(bServiceExecutableUpdated, TEXT("Failed to copy zenserver to install location '%s'."), *InstallServicePath);
+
+		bool bUtilityExecutableUpdated = AttemptFileCopyWithRetries(*InstallUtilityPath, *InTreeUtilityPath, 5.0);
+		checkf(bUtilityExecutableUpdated, TEXT("Failed to copy zen to install location '%s'."), *InstallUtilityPath);
+
+		bMainExecutablesUpdated = (bServiceExecutableUpdated && bUtilityExecutableUpdated);
+		if (bMainExecutablesUpdated)
+		{
+			AttemptFileCopyWithRetries(*InstallVersionCache, *InTreeVersionCache, 1.0);
+		}
+	}
+
+#if PLATFORM_WINDOWS
+	FString InTreeSymbolFilePath = FPaths::SetExtension(InTreeServicePath, TEXT("pdb"));
+	FString InstallSymbolFilePath = FPaths::SetExtension(InstallServicePath, TEXT("pdb"));
+
+	if (FileManager.FileExists(*InTreeSymbolFilePath) && (bMainExecutablesUpdated || !FileManager.FileExists(*InstallSymbolFilePath)))
+	{
+		AttemptFileCopyWithRetries(*InstallSymbolFilePath, *InTreeSymbolFilePath, 1.0);
+	}
+#endif
+
+	FString InTreeCrashpadHandlerFilePath = FPaths::ConvertRelativePathToFull(FPlatformProcess::GenerateApplicationPath(TEXT("crashpad_handler"), EBuildConfiguration::Development));
+	FString InstallCrashpadHandlerFilePath = FPaths::ConvertRelativePathToFull(FPaths::Combine(FPlatformProcess::ApplicationSettingsDir(), TEXT("Zen\\Install"), FString(FPathViews::GetCleanFilename(InTreeCrashpadHandlerFilePath))));
+
+	if (FileManager.FileExists(*InTreeCrashpadHandlerFilePath) && (bMainExecutablesUpdated || !FileManager.FileExists(*InstallCrashpadHandlerFilePath)))
+	{
+		AttemptFileCopyWithRetries(*InstallCrashpadHandlerFilePath, *InTreeCrashpadHandlerFilePath, 1.0);
+	}
+
+	return InstallServicePath;
+}
+
+bool
+FZenServiceInstance::AutoLaunch(const FServiceAutoLaunchSettings& InSettings, FString&& ExecutablePath, FString& OutHostName, uint16& OutPort)
+{
+	int16 DesiredPort = InSettings.DesiredPort;
+	IFileManager& FileManager = IFileManager::Get();
+	const FString LockFilePath = FPaths::Combine(InSettings.DataPath, TEXT(".lock"));
+	const FString ExecutionContextFilePath = FPaths::SetExtension(ExecutablePath, TEXT(".runcontext"));
+
+	FString WorkingDirectory = FPaths::GetPath(ExecutablePath);
+
+	bool bReUsingExistingInstance = false;
+
+	if (IsLockFileLocked(*LockFilePath, true))
+	{
+		// If an instance is running with this data path, check if we can use it and what port it is on
+		uint16 CurrentPort = 0;
+		FCbObject LockObject;
+		if (ReadCbLockFile(LockFilePath, LockObject))
+		{
+			bool bIsReady = LockObject["ready"].AsBool();
+			if (bIsReady)
+			{
+				CurrentPort = LockObject["port"].AsUInt16();
+			}
+		}
+
+		bool bCurrentInstanceUsable = false;
+
+		FZenLocalServiceRunContext DesiredRunContext;
+		DesiredRunContext.Executable = ExecutablePath;
+		DesiredRunContext.CommandlineArguments = DetermineCmdLineWithoutTransientComponents(InSettings, CurrentPort);
+		DesiredRunContext.WorkingDirectory = WorkingDirectory;
+		DesiredRunContext.DataPath = InSettings.DataPath;
+		DesiredRunContext.bShowConsole = InSettings.bShowConsole;
+
+		FZenLocalServiceRunContext CurrentRunContext;
+		if (CurrentRunContext.ReadFromJsonFile(*ExecutionContextFilePath) && (DesiredRunContext == CurrentRunContext))
+		{
+			DesiredPort = CurrentPort;
+			bReUsingExistingInstance = true;
+		}
+		else
+		{
+			RequestZenShutdownOnPort(CurrentPort);
+			WaitForZenShutdown(*LockFilePath, 5.0);
+		}
+	}
+
+	if (!bReUsingExistingInstance)
+	{
+		RequestZenShutdownOnPort(DesiredPort);
+	}
+
+
+	bool bProcessIsLive = IsLockFileLocked(*LockFilePath);
+
+	// When limiting process lifetime, always re-launch to add sponsor process IDs.
+	// When not limiting process lifetime, only launch if the process is not already live.
+	if (InSettings.bLimitProcessLifetime || !bProcessIsLive)
+	{
+		FString ParmsWithoutTransients = DetermineCmdLineWithoutTransientComponents(InSettings, DesiredPort);
+		FString TransientParms;
+
+		if (InSettings.bLimitProcessLifetime)
+		{
+			TransientParms.Appendf(TEXT("--owner-pid %d"), FPlatformProcess::GetCurrentProcessId());
+		}
+
+		FZenLocalServiceRunContext EffectiveRunContext;
+		EffectiveRunContext.Executable = ExecutablePath;
+		EffectiveRunContext.CommandlineArguments = ParmsWithoutTransients;
+		EffectiveRunContext.WorkingDirectory = WorkingDirectory;
+		EffectiveRunContext.DataPath = InSettings.DataPath;
+		EffectiveRunContext.bShowConsole = InSettings.bShowConsole;
+
+		FProcHandle Proc = StartLocalService(EffectiveRunContext, TransientParms.IsEmpty() ? nullptr : *TransientParms);
+		if (!bProcessIsLive)
+		{
+			EffectiveRunContext.WriteToJsonFile(*ExecutionContextFilePath);
+		}
+
+		bProcessIsLive = Proc.IsValid();
+	}
+
+
+	OutHostName = TEXT("[::1]");
+	// Default to assuming that we get to run on the port we want
+	OutPort = DesiredPort;
+
+	if (bProcessIsLive)
+	{
+
+		FScopedSlowTask WaitForZenReadySlowTask(0, NSLOCTEXT("Zen", "Zen_WaitingForReady", "Waiting for ZenServer to be ready"));
+		uint64 ZenWaitStartTime = FPlatformTime::Cycles64();
+		enum class EWaitDurationPhase
+		{
+			Short,
+			Medium,
+			Long
+		} DurationPhase = EWaitDurationPhase::Short;
+		bool bIsReady = false;
+		while (!bIsReady)
+		{
+			FCbObject LockObject;
+			if (ReadCbLockFile(LockFilePath, LockObject))
+			{
+				bIsReady = LockObject["ready"].AsBool();
+				if (bIsReady)
+				{
+					OutPort = LockObject["port"].AsUInt16(DesiredPort);
+					break;
+				}
+			}
+
+			double ZenWaitDuration = FPlatformTime::ToSeconds64(FPlatformTime::Cycles64() - ZenWaitStartTime);
+			if (ZenWaitDuration < 10.0)
+			{
+				// Initial 10 second window of higher frequency checks
+				FPlatformProcess::Sleep(0.01f);
+			}
+			else
+			{
+				if (DurationPhase == EWaitDurationPhase::Short)
+				{
+					if (!IsLockFileLocked(*LockFilePath))
+					{
+						if (FApp::IsUnattended())
+						{
+							checkf(false, TEXT("ZenServer did not launch in the expected duration."));
+						}
+						else
+						{
+							FText ZenLaunchFailurePromptTitle = NSLOCTEXT("Zen", "Zen_LaunchFailurePromptTitle", "Failed to launch");
+
+							FFormatNamedArguments FormatArguments;
+							FString LogFilePath = FPaths::Combine(InSettings.DataPath, TEXT("logs"), TEXT("zenserver.log"));
+							FPaths::MakePlatformFilename(LogFilePath);
+							FormatArguments.Add(TEXT("LogFilePath"), FText::FromString(LogFilePath));
+							FText ZenLaunchFailurePromptText = FText::Format(NSLOCTEXT("Zen", "Zen_LaunchFailurePromptText", "ZenServer failed to launch. This process will now exit. Please check the ZenServer log file for details:\n{LogFilePath}"), FormatArguments);
+							FPlatformMisc::MessageBoxExt(EAppMsgType::Ok, *ZenLaunchFailurePromptText.ToString(), *ZenLaunchFailurePromptTitle.ToString());
+							FPlatformMisc::RequestExit(true);
+							return false;
+						}
+					}
+					// Note that the dialog may not show up when zenserver is needed early in the launch cycle, but this will at least ensure
+					// the splash screen is refreshed with the appropriate text status message.
+					WaitForZenReadySlowTask.MakeDialog(true, false);
+					UE_LOG(LogZenServiceInstance, Display, TEXT("Waiting for ZenServer to be ready..."));
+					DurationPhase = EWaitDurationPhase::Medium;
+				}
+				else if (!FApp::IsUnattended() && ZenWaitDuration > 20.0 && (DurationPhase == EWaitDurationPhase::Medium))
+				{
+					FText ZenLongWaitPromptTitle = NSLOCTEXT("Zen", "Zen_LongWaitPromptTitle", "Wait for ZenServer?");
+					FText ZenLongWaitPromptText = NSLOCTEXT("Zen", "Zen_LongWaitPromptText", "ZenServer is taking a long time to launch. It may be performing maintenance. Keep waiting?");
+					if (FPlatformMisc::MessageBoxExt(EAppMsgType::YesNo, *ZenLongWaitPromptText.ToString(), *ZenLongWaitPromptTitle.ToString()) == EAppReturnType::No)
+					{
+						FPlatformMisc::RequestExit(true);
+						return false;
+					}
+					DurationPhase = EWaitDurationPhase::Long;
+				}
+
+				if (WaitForZenReadySlowTask.ShouldCancel())
+				{
+					FPlatformMisc::RequestExit(true);
+					return false;
+				}
+				FPlatformProcess::Sleep(0.1f);
+			}
+		}
+
+		return bIsReady;
+	}
+	else
+	{
+		return false;
+	}
+	return true;
+}
+
+bool 
+FZenServiceInstance::GetStats(FZenStats& Stats)
+{
+	check(IsInGameThread());
+
+	// If we've already requested a stats and they are ready then grab them
+	if ( StatsRequest.IsReady() == true )
+	{
+		LastStats		= StatsRequest.Get();
+		LastStatsTime	= FPlatformTime::Cycles64();
+
+		StatsRequest.Reset();
+	}
+	
+	// Make a copy of the last updated stats
+	Stats = LastStats;
+
+	const uint64 CurrentTime = FPlatformTime::Cycles64();
+	constexpr double MinTimeBetweenRequestsInSeconds = 0.5;
+	const double DeltaTimeInSeconds = FPlatformTime::ToSeconds64(CurrentTime - LastStatsTime);
+
+	if (!StatsRequest.IsValid() && DeltaTimeInSeconds > MinTimeBetweenRequestsInSeconds)
+	{
+#if WITH_EDITOR
+		EAsyncExecution ThreadPool = EAsyncExecution::LargeThreadPool;
+#else
+		EAsyncExecution ThreadPool = EAsyncExecution::ThreadPool;
+#endif
+		if (!StatsHttpRequest.IsValid())
+		{
+			TStringBuilder<128> ZenDomain;
+			ZenDomain << HostName << TEXT(":") << Port;
+			StatsHttpRequest = MakePimpl<FZenHttpRequest>(ZenDomain.ToString(), false);
+		}
+
+		// We've not got any requests in flight and we've met a given time requirement for requests
+		StatsRequest = Async(ThreadPool, [this]
+			{
+				UE::Zen::FZenHttpRequest& Request = *StatsHttpRequest.Get();
+				Request.Reset();
+
+				TArray64<uint8> GetBuffer;
+				FZenHttpRequest::Result Result = Request.PerformBlockingDownload(TEXTVIEW("/stats/z$"), &GetBuffer, Zen::EContentType::CbObject);
+
+				FZenStats Stats;
+
+				if (Result == Zen::FZenHttpRequest::Result::Success && Request.GetResponseCode() == 200)
+				{
+					FCbObjectView RootObjectView(GetBuffer.GetData());
+
+					FCbObjectView RequestsObjectView = RootObjectView["requests"].AsObjectView();
+					FZenRequestStats& RequestStats = Stats.RequestStats;
+
+					RequestStats.Count = RequestsObjectView["count"].AsInt64();
+					RequestStats.RateMean = RequestsObjectView["rate_mean"].AsDouble();
+					RequestStats.TAverage = RequestsObjectView["t_avg"].AsDouble();
+					RequestStats.TMin = RequestsObjectView["t_min"].AsDouble();
+					RequestStats.TMax = RequestsObjectView["t_max"].AsDouble();
+
+					FCbObjectView CacheObjectView = RootObjectView["cache"].AsObjectView();
+					FZenCacheStats& CacheStats = Stats.CacheStats;
+
+					CacheStats.Hits = CacheObjectView["hits"].AsInt64();
+					CacheStats.Misses = CacheObjectView["misses"].AsInt64();
+					CacheStats.HitRatio = CacheObjectView["hit_ratio"].AsDouble();
+					CacheStats.UpstreamHits = CacheObjectView["upstream_hits"].AsInt64();
+					CacheStats.UpstreamRatio = CacheObjectView["upstream_ratio"].AsDouble();
+
+					FCbObjectView CacheSizeObjectView = CacheObjectView["size"].AsObjectView();
+					FZenCacheSizeStats& CacheSizeStats = CacheStats.Size;
+					CacheSizeStats.Disk = CacheSizeObjectView["disk"].AsDouble();
+					CacheSizeStats.Memory = CacheSizeObjectView["memory"].AsDouble();
+
+					FCbObjectView UpstreamObjectView = RootObjectView["upstream"].AsObjectView();
+					FZenUpstreamStats& UpstreamStats = Stats.UpstreamStats;
+
+					UpstreamStats.Reading = UpstreamObjectView["reading"].AsBool();
+					UpstreamStats.Writing = UpstreamObjectView["writing"].AsBool();
+					UpstreamStats.WorkerThreads = UpstreamObjectView["worker_threads"].AsInt64();
+					UpstreamStats.QueueCount = UpstreamObjectView["queue_count"].AsInt64();
+					UpstreamStats.TotalUploadedMB = 0.0;
+					UpstreamStats.TotalDownloadedMB = 0.0;
+
+					FCbObjectView UpstreamRequestObjectView = RootObjectView["upstream_gets"].AsObjectView();
+					FZenRequestStats& UpstreamRequestStats = Stats.UpstreamRequestStats;
+
+					UpstreamRequestStats.Count = UpstreamRequestObjectView["count"].AsInt64();
+					UpstreamRequestStats.RateMean = UpstreamRequestObjectView["rate_mean"].AsDouble();
+					UpstreamRequestStats.TAverage = UpstreamRequestObjectView["t_avg"].AsDouble();
+					UpstreamRequestStats.TMin = UpstreamRequestObjectView["t_min"].AsDouble();
+					UpstreamRequestStats.TMax = UpstreamRequestObjectView["t_max"].AsDouble();
+
+					FCbArrayView EndpPointArrayView = UpstreamObjectView["endpoints"].AsArrayView();
+
+					for (FCbFieldView FieldView : EndpPointArrayView)
+					{
+						FCbObjectView EndPointView = FieldView.AsObjectView();
+						FZenEndPointStats EndPointStats;
+
+						EndPointStats.Name = FString(EndPointView["name"].AsString());
+						EndPointStats.Url = FString(EndPointView["url"].AsString());
+						EndPointStats.Health = FString(EndPointView["state"].AsString());
+
+						if (FCbObjectView Cache = EndPointView["cache"].AsObjectView())
+						{
+							EndPointStats.HitRatio = Cache["hit_ratio"].AsDouble();
+							EndPointStats.UploadedMB = Cache["put_bytes"].AsDouble() / 1024.0 / 1024.0;
+							EndPointStats.DownloadedMB = Cache["get_bytes"].AsDouble() / 1024.0 / 1024.0;
+							EndPointStats.ErrorCount = Cache["error_count"].AsInt64();
+						}
+
+						UpstreamStats.TotalUploadedMB += EndPointStats.UploadedMB;
+						UpstreamStats.TotalDownloadedMB += EndPointStats.DownloadedMB;
+
+						UpstreamStats.EndPointStats.Push(EndPointStats);
+					}
+
+					FCbObjectView CASObjectView = RootObjectView["cid"].AsObjectView();
+					FCbObjectView CASSizeObjectView = CASObjectView["size"].AsObjectView();
+
+					FZenCASSizeStats& CASSizeStats = Stats.CASStats.Size;
+
+					CASSizeStats.Tiny = CASSizeObjectView["tiny"].AsInt64();
+					CASSizeStats.Small = CASSizeObjectView["small"].AsInt64();
+					CASSizeStats.Large = CASSizeObjectView["large"].AsInt64();
+					CASSizeStats.Total = CASSizeObjectView["total"].AsInt64();
+
+					Stats.IsValid = true;
+				}
+
+				return Stats;
+			});
+	}
+
+	return Stats.IsValid;
+}
+
+bool 
+FZenServiceInstance::GetGCStatus(FGCStatus& Status)
+{
+	check(IsInGameThread());
+
+	// If we've already requested status and it is ready then grab it
+	if (GCStatusRequest.IsReady() == true )
+	{
+		LastGCStatus	 = GCStatusRequest.Get();
+		LastGCStatusTime = FPlatformTime::Cycles64();
+
+		GCStatusRequest.Reset();
+	}
+	
+	// Make a copy of the last updated status
+	if (LastGCStatus.IsSet())
+	{
+		Status = LastGCStatus.GetValue();
+	}
+
+	const uint64 CurrentTime = FPlatformTime::Cycles64();
+	constexpr double MinTimeBetweenRequestsInSeconds = 0.5;
+	const double DeltaTimeInSeconds = FPlatformTime::ToSeconds64(CurrentTime - LastGCStatusTime);
+
+	if (!GCStatusRequest.IsValid() && DeltaTimeInSeconds > MinTimeBetweenRequestsInSeconds)
+	{
+#if WITH_EDITOR
+		EAsyncExecution ThreadPool = EAsyncExecution::LargeThreadPool;
+#else
+		EAsyncExecution ThreadPool = EAsyncExecution::ThreadPool;
+#endif
+		if (!GCStatusHttpRequest.IsValid())
+		{
+			TStringBuilder<128> ZenDomain;
+			ZenDomain << HostName << TEXT(":") << Port;
+			GCStatusHttpRequest = MakePimpl<FZenHttpRequest>(ZenDomain.ToString(), false);
+		}
+
+		// We've not got any requests in flight and we've met a given time requirement for requests
+		GCStatusRequest = Async(ThreadPool, [this]
+			{
+				UE::Zen::FZenHttpRequest& Request = *GCStatusHttpRequest.Get();
+				Request.Reset();
+
+				TArray64<uint8> GetBuffer;
+				FZenHttpRequest::Result Result = Request.PerformBlockingDownload(TEXTVIEW("/admin/gc"), &GetBuffer, Zen::EContentType::CbObject);
+
+				TOptional<FGCStatus> GCStatus;
+
+				if (Result == Zen::FZenHttpRequest::Result::Success && Request.GetResponseCode() == 200)
+				{
+					FCbObjectView RootObjectView(GetBuffer.GetData());
+
+					GCStatus.Emplace();
+					GCStatus->Description = FString(RootObjectView["Status"].AsString());
+				}
+
+				return GCStatus;
+			});
+	}
+
+	return LastGCStatus.IsSet();
+}
+
+bool 
+FZenServiceInstance::RequestGC(const bool* OverrideCollectSmallObjects, const uint32* OverrideMaxCacheDuration)
+{
+	TStringBuilder<128> ZenDomain;
+	ZenDomain << HostName << TEXT(":") << Port;
+	UE::Zen::FZenHttpRequest Request(ZenDomain.ToString(), false);
+
+	TCHAR Separators[] = {TEXT('?'), TEXT('&')};
+	int32 SeparatorIndex = 0;
+	TStringBuilder<128> Query;
+	Query << TEXTVIEW("/admin/gc");
+
+	if (OverrideCollectSmallObjects)
+	{
+		Query << Separators[SeparatorIndex] << LexToString(*OverrideCollectSmallObjects);
+		SeparatorIndex = FMath::Min(SeparatorIndex + 1, (int32)UE_ARRAY_COUNT(Separators));
+	}
+
+	if (OverrideMaxCacheDuration)
+	{
+		Query << Separators[SeparatorIndex] << LexToString(*OverrideMaxCacheDuration);
+		SeparatorIndex = FMath::Min(SeparatorIndex + 1, (int32)UE_ARRAY_COUNT(Separators));
+	}
+
+	FZenHttpRequest::Result Result = Request.PerformBlockingPost(Query.ToString(), FMemoryView());
+
+	if (Result == Zen::FZenHttpRequest::Result::Success && Request.GetResponseCode() == 200)
+	{
+		FCbObjectView ResponseObject = FCbObjectView(Request.GetResponseBuffer().GetData());
+		FUtf8StringView ResponseStatus = ResponseObject["status"].AsString();
+
+		return (ResponseStatus == "Started") || (ResponseStatus == "Running");
+	}
+	return false;
+}
+
+void FZenServiceInstance::OnAnalyticsEvent(TArray<FAnalyticsEventAttribute>& Attributes )
+{
+	// Grab the Zen summary stats
+	FZenStats ZenStats;
+
+	GetStats(ZenStats);
+
+	const FString BaseName = TEXT("Zen");
+
+	{
+		FString AttrName = BaseName + TEXT(".Enabled");
+		Attributes.Emplace(MoveTemp(AttrName), ZenStats.IsValid);
+	}
+
+	{
+		FString AttrName = BaseName + TEXT(".Cache.HitRatio");
+		Attributes.Emplace(MoveTemp(AttrName), ZenStats.CacheStats.HitRatio);
+	}
+
+	{
+		FString AttrName = BaseName + TEXT(".Cache.Hits");
+		Attributes.Emplace(MoveTemp(AttrName), ZenStats.CacheStats.Hits);
+	}
+
+	{
+		FString AttrName = BaseName + TEXT(".Cache.Misses");
+		Attributes.Emplace(MoveTemp(AttrName), ZenStats.CacheStats.Misses);
+	}
+
+	{
+		FString AttrName = BaseName + TEXT(".Cache.Size.Disk");
+		Attributes.Emplace(MoveTemp(AttrName), ZenStats.CacheStats.Size.Disk);
+	}
+
+	{
+		FString AttrName = BaseName + TEXT(".Cache.Size.Memory");
+		Attributes.Emplace(MoveTemp(AttrName), ZenStats.CacheStats.Size.Memory);
+	}
+
+	{
+		FString AttrName = BaseName + TEXT(".Cache.UpstreamHits");
+		Attributes.Emplace(MoveTemp(AttrName), ZenStats.CacheStats.UpstreamHits);
+	}
+
+	{
+		FString AttrName = BaseName + TEXT(".Cache.UpstreamRatio");
+		Attributes.Emplace(MoveTemp(AttrName), ZenStats.CacheStats.UpstreamRatio);
+	}
+
+	{
+		FString AttrName = BaseName + TEXT(".Cache.TotalUploadedMB");
+		Attributes.Emplace(MoveTemp(AttrName), ZenStats.UpstreamStats.TotalUploadedMB);
+	}
+
+	{
+		FString AttrName = BaseName + TEXT(".Upstream.TotalDownloadedMB");
+		Attributes.Emplace(MoveTemp(AttrName), ZenStats.UpstreamStats.TotalDownloadedMB);
+	}
+
+	{
+		FString AttrName = BaseName + TEXT(".Upstream.TotalUploadedMB");
+		Attributes.Emplace(MoveTemp(AttrName), ZenStats.UpstreamStats.TotalUploadedMB);
+	}
+
+	{
+		FString AttrName = BaseName + TEXT(".Cas.Size.Large");
+		Attributes.Emplace(MoveTemp(AttrName), ZenStats.CASStats.Size.Large);
+	}
+
+	{
+		FString AttrName = BaseName + TEXT(".Cas.Size.Small");
+		Attributes.Emplace(MoveTemp(AttrName), ZenStats.CASStats.Size.Small);
+	}
+
+	{
+		FString AttrName = BaseName + TEXT(".Cas.Size.Tiny");
+		Attributes.Emplace(MoveTemp(AttrName), ZenStats.CASStats.Size.Tiny);
+	}
+
+	{
+		FString AttrName = BaseName + TEXT(".Cas.Size.Total");
+		Attributes.Emplace(MoveTemp(AttrName), ZenStats.CASStats.Size.Total);
+	}
+}
+
+#endif // UE_WITH_ZEN
+
+}
+
diff --git a/Engine/Source/Developer/Zen/Private/ZenVersion.cpp b/Engine/Source/Developer/Zen/Private/ZenVersion.cpp
new file mode 100644
--- /dev/null
+++ b/Engine/Source/Developer/Zen/Private/ZenVersion.cpp
@@ -0,0 +1,114 @@
+// Copyright Epic Games, Inc. All Rights Reserved.
+#include "ZenVersion.h"
+#include "Containers/StringFwd.h"
+#include "Misc/StringBuilder.h"
+
+namespace UE::Zen
+{
+
+void FZenVersion::Reset()
+{
+	MajorVersion = 0;
+	MinorVersion = 0;
+	PatchVersion = 0;
+	Details.Empty();
+}
+
+bool FZenVersion::TryParse(const TCHAR* InString)
+{
+	Reset();
+
+	if (InString == nullptr)
+	{
+		return false;
+	}
+
+	const TCHAR* CurrentSegment = InString;
+	uint32 CurrentSegmentIndex;
+	for (CurrentSegmentIndex = 0; CurrentSegmentIndex < 3; ++CurrentSegmentIndex)
+	{
+		TCHAR* End = nullptr;
+		uint64 ParsedValue = FCString::Strtoui64(CurrentSegment, &End, 10);
+
+		if (End == CurrentSegment)
+		{
+			// A segment with zero parsable characters is an error
+			break;
+		}
+
+		switch (CurrentSegmentIndex)
+		{
+		case 0:
+			MajorVersion = (uint32)ParsedValue;
+			break;
+		case 1:
+			MinorVersion = (uint32)ParsedValue;
+			break;
+		case 2:
+			PatchVersion = (uint32)ParsedValue;
+			break;
+		}
+
+		if (*End != TCHAR('.'))
+		{
+			CurrentSegment = End;
+			if (*End == TCHAR('-'))
+			{
+				Details = End + 1;
+				CurrentSegment = End + 1 + Details.Len();
+			}
+			++CurrentSegmentIndex;
+			break;
+		}
+		CurrentSegment = End + 1;
+	}
+
+	// Successful parsing is determined by whether we managed to parse at least 1 numeric segment
+	// and we've parsed to the end of the given string.  Details are optional.
+	return (CurrentSegmentIndex > 0) && (*CurrentSegment == TEXT('\0'));
+}
+
+FString FZenVersion::ToString(bool bDetailed) const
+{
+	TStringBuilder<64> Builder;
+	Builder << MajorVersion;
+	Builder.AppendChar(TEXT('.'));
+	Builder << MinorVersion;
+	Builder.AppendChar(TEXT('.'));
+	Builder << PatchVersion;
+	if (bDetailed && !Details.IsEmpty())
+	{
+		Builder.AppendChar(TEXT('-'));
+		Builder << Details;
+	}
+	return Builder.ToString();
+}
+
+bool FZenVersion::operator<(FZenVersion& Other) const
+{
+	if (MajorVersion < Other.MajorVersion)
+	{
+		return true;
+	}
+	else if (MajorVersion == Other.MajorVersion)
+	{
+		if (MinorVersion < Other.MinorVersion)
+		{
+			return true;
+		}
+		else if (MinorVersion == Other.MinorVersion)
+		{
+			if (PatchVersion < Other.PatchVersion)
+			{
+				return true;
+			}
+			else if (PatchVersion == Other.PatchVersion)
+			{
+				return Details < Other.Details;
+			}
+		}
+	}
+	return false;
+}
+
+} // namespace UE::Zen
