Move Zen module public headers into an "Experimental" subfolder to have them reflect the feature status.
Fix locations that include the headers to use the "Experimental" subfolder and re-sort includes.  No functional change, purely a move, search+replace, and sort of includes.

diff --git a/Engine/Source/Developer/DerivedDataCache/Private/DerivedDataCache.cpp b/Engine/Source/Developer/DerivedDataCache/Private/DerivedDataCache.cpp
--- a/Engine/Source/Developer/DerivedDataCache/Private/DerivedDataCache.cpp
+++ b/Engine/Source/Developer/DerivedDataCache/Private/DerivedDataCache.cpp
@@ -22,6 +22,7 @@
 #include "DerivedDataRequestOwner.h"
 #include "DerivedDataThreadPoolTask.h"
 #include "Experimental/Async/LazyEvent.h"
+#include "Experimental/ZenServerInterface.h"
 #include "Features/IModularFeatures.h"
 #include "HAL/ThreadSafeCounter.h"
 #include "Misc/CommandLine.h"
@@ -33,7 +34,6 @@
 #include "Serialization/CompactBinaryWriter.h"
 #include "Stats/Stats.h"
 #include "Stats/StatsMisc.h"
-#include "ZenServerInterface.h"
 #include <atomic>
 
 DEFINE_STAT(STAT_DDC_NumGets);
diff --git a/Engine/Source/Developer/DerivedDataCache/Private/Tests/HttpCacheStoreTest.cpp b/Engine/Source/Developer/DerivedDataCache/Private/Tests/HttpCacheStoreTest.cpp
--- a/Engine/Source/Developer/DerivedDataCache/Private/Tests/HttpCacheStoreTest.cpp
+++ b/Engine/Source/Developer/DerivedDataCache/Private/Tests/HttpCacheStoreTest.cpp
@@ -7,13 +7,13 @@
 #include "DerivedDataRequestOwner.h"
 #include "DerivedDataValue.h"
 #include "DerivedDataValueId.h"
+#include "Experimental/ZenServerInterface.h"
 #include "HAL/FileManager.h"
 #include "Memory/CompositeBuffer.h"
 #include "Misc/AutomationTest.h"
 #include "Misc/Paths.h"
 #include "Misc/SecureHash.h"
 #include "Serialization/CompactBinaryWriter.h"
-#include "ZenServerInterface.h"
 
 // Test is targeted at HttpDerivedDataBackend but with some backend test interface it could be generalized
 // to function against all backends.
diff --git a/Engine/Source/Developer/DerivedDataCache/Private/ZenCacheStore.cpp b/Engine/Source/Developer/DerivedDataCache/Private/ZenCacheStore.cpp
--- a/Engine/Source/Developer/DerivedDataCache/Private/ZenCacheStore.cpp
+++ b/Engine/Source/Developer/DerivedDataCache/Private/ZenCacheStore.cpp
@@ -1,7 +1,7 @@
 // Copyright Epic Games, Inc. All Rights Reserved.
 
-#include "ZenServerInterface.h"
 #include "DerivedDataLegacyCacheStore.h"
+#include "Experimental/ZenServerInterface.h"
 
 #if UE_WITH_ZEN
 
@@ -13,6 +13,7 @@
 #include "DerivedDataChunk.h"
 #include "DerivedDataRequest.h"
 #include "DerivedDataRequestOwner.h"
+#include "Experimental/ZenStatistics.h"
 #include "Http/HttpClient.h"
 #include "Math/UnrealMathUtility.h"
 #include "Misc/App.h"
@@ -28,7 +29,6 @@
 #include "Templates/Function.h"
 #include "ZenBackendUtils.h"
 #include "ZenSerialization.h"
-#include "ZenStatistics.h"
 
 TRACE_DECLARE_INT_COUNTER(ZenDDC_Get,			TEXT("ZenDDC Get"));
 TRACE_DECLARE_INT_COUNTER(ZenDDC_GetHit,		TEXT("ZenDDC Get Hit"));
diff --git a/Engine/Source/Developer/IoStoreUtilities/Internal/ZenStoreHttpClient.h b/Engine/Source/Developer/IoStoreUtilities/Internal/ZenStoreHttpClient.h
--- a/Engine/Source/Developer/IoStoreUtilities/Internal/ZenStoreHttpClient.h
+++ b/Engine/Source/Developer/IoStoreUtilities/Internal/ZenStoreHttpClient.h
@@ -5,13 +5,13 @@
 #include "Async/Future.h"
 #include "Containers/Array.h"
 #include "Containers/StringView.h"
+#include "Experimental/ZenServerInterface.h"
 #include "HAL/CriticalSection.h"
 #include "HAL/Platform.h"
 #include "IO/IoDispatcher.h"
 #include "Misc/StringBuilder.h"
 #include "Serialization/CompactBinaryPackage.h"
 #include "Templates/UniquePtr.h"
-#include "ZenServerInterface.h"
 
 class FCbPackage;
 class FCbObject;
diff --git a/Engine/Source/Developer/IoStoreUtilities/Public/PackageStoreManifest.h b/Engine/Source/Developer/IoStoreUtilities/Public/PackageStoreManifest.h
--- a/Engine/Source/Developer/IoStoreUtilities/Public/PackageStoreManifest.h
+++ b/Engine/Source/Developer/IoStoreUtilities/Public/PackageStoreManifest.h
@@ -6,13 +6,13 @@
 #include "Containers/ContainerAllocationPolicies.h"
 #include "Containers/Map.h"
 #include "Containers/UnrealString.h"
+#include "Experimental/ZenServerInterface.h"
 #include "HAL/CriticalSection.h"
 #include "HAL/Platform.h"
 #include "HAL/PlatformCrt.h"
 #include "IO/IoDispatcher.h"
 #include "Templates/UniquePtr.h"
 #include "UObject/NameTypes.h"
-#include "ZenServerInterface.h"
 
 class FPackageStoreManifest
 {
diff --git a/Engine/Source/Developer/Zen/Internal/ZenSerialization.h b/Engine/Source/Developer/Zen/Internal/ZenSerialization.h
--- a/Engine/Source/Developer/Zen/Internal/ZenSerialization.h
+++ b/Engine/Source/Developer/Zen/Internal/ZenSerialization.h
@@ -2,7 +2,7 @@
 
 #pragma once
 
-#include "ZenServerInterface.h"
+#include "Experimental/ZenServerInterface.h"
 
 #if UE_WITH_ZEN
 
diff --git a/Engine/Source/Developer/Zen/Internal/ZenServerHttp.h b/Engine/Source/Developer/Zen/Internal/ZenServerHttp.h
--- a/Engine/Source/Developer/Zen/Internal/ZenServerHttp.h
+++ b/Engine/Source/Developer/Zen/Internal/ZenServerHttp.h
@@ -2,11 +2,10 @@
 
 #pragma once
 
-#include "ZenGlobals.h"
-
 #include "Containers/StaticArray.h"
 #include "Containers/StringView.h"
 #include "Containers/UnrealString.h"
+#include "Experimental/ZenGlobals.h"
 #include "Memory/MemoryFwd.h"
 
 // if there is a platform-specific include then it must be used in the header file in case it defines CURL_STRICTER
diff --git a/Engine/Source/Developer/Zen/Private/ZenServerInterface.cpp b/Engine/Source/Developer/Zen/Private/ZenServerInterface.cpp
--- a/Engine/Source/Developer/Zen/Private/ZenServerInterface.cpp
+++ b/Engine/Source/Developer/Zen/Private/ZenServerInterface.cpp
@@ -1,6 +1,6 @@
 // Copyright Epic Games, Inc. All Rights Reserved.
-#include "ZenServerInterface.h"
 
+#include "Experimental/ZenServerInterface.h"
 
 #include "ZenBackendUtils.h"
 #include "ZenSerialization.h"
diff --git a/Engine/Source/Developer/Zen/Public/Experimental/ZenGlobals.h b/Engine/Source/Developer/Zen/Public/Experimental/ZenGlobals.h
new file mode 100644
--- /dev/null
+++ b/Engine/Source/Developer/Zen/Public/Experimental/ZenGlobals.h
@@ -0,0 +1,13 @@
+// Copyright Epic Games, Inc. All Rights Reserved.
+
+#pragma once
+
+#include "HAL/Platform.h"
+
+#ifndef UE_WITH_ZEN
+#	if PLATFORM_DESKTOP
+#		define UE_WITH_ZEN 1
+#	else
+#		define UE_WITH_ZEN 0
+#	endif
+#endif
diff --git a/Engine/Source/Developer/Zen/Public/Experimental/ZenServerInterface.h b/Engine/Source/Developer/Zen/Public/Experimental/ZenServerInterface.h
new file mode 100644
--- /dev/null
+++ b/Engine/Source/Developer/Zen/Public/Experimental/ZenServerInterface.h
@@ -0,0 +1,216 @@
+// Copyright Epic Games, Inc. All Rights Reserved.
+
+#pragma once
+
+#include "Async/Future.h"
+#include "Containers/StringFwd.h"
+#include "Containers/StringView.h"
+#include "Containers/UnrealString.h"
+#include "Dom/JsonObject.h"
+#include "Experimental/ZenGlobals.h"
+#include "HAL/Platform.h"
+#include "HAL/PlatformCrt.h"
+#include "Misc/Optional.h"
+#include "Misc/TVariant.h"
+#include "Policies/PrettyJsonPrintPolicy.h"
+#include "Serialization/JsonWriter.h"
+#include "Templates/PimplPtr.h"
+#include "Templates/UniquePtr.h"
+
+#if UE_WITH_ZEN
+#	include "ZenStatistics.h"
+#endif
+
+#define UE_API ZEN_API
+
+struct FAnalyticsEventAttribute;
+
+namespace UE::Zen
+{
+
+struct FServiceConnectSettings
+{
+	FString HostName;
+	uint16 Port = 1337;
+};
+
+struct FServiceAutoLaunchSettings
+{
+	FString DataPath;
+	FString ExtraArgs;
+	uint16 DesiredPort = 1337;
+	bool bShowConsole = false;
+	bool bLimitProcessLifetime = false;
+};
+
+struct FServiceSettings
+{
+	TVariant<FServiceAutoLaunchSettings, FServiceConnectSettings> SettingsVariant;
+
+	inline bool IsAutoLaunch() const { return SettingsVariant.IsType<FServiceAutoLaunchSettings>(); }
+	inline bool IsConnectExisting() const { return SettingsVariant.IsType<FServiceConnectSettings>(); }
+
+	UE_API void ReadFromConfig();
+	UE_API void ReadFromJson(FJsonObject& JsonObject);
+	UE_API void ReadFromURL(FStringView InstanceURL);
+
+	UE_API void WriteToJson(TJsonWriter<TCHAR, TPrettyJsonPrintPolicy<TCHAR>>& Writer) const;
+
+private:
+	bool TryApplyAutoLaunchOverride();
+};
+
+}
+
+#if UE_WITH_ZEN
+
+namespace UE::Zen
+{
+
+class FZenServiceInstance;
+class FZenLocalServiceRunContext;
+
+UE_API bool TryGetLocalServiceRunContext(FZenLocalServiceRunContext& OutContext);
+
+class FZenLocalServiceRunContext
+{
+public:
+	FString GetExecutable() const { return Executable; }
+	FString GetCommandlineArguments() const { return CommandlineArguments; }
+	FString GetWorkingDirectory() const { return WorkingDirectory; }
+	FString GetDataPath() const { return DataPath; }
+	bool GetShowConsole() const { return bShowConsole; }
+
+	friend inline bool operator==(FZenLocalServiceRunContext LHS, FZenLocalServiceRunContext RHS)
+	{
+		return (LHS.Executable == RHS.Executable) && (LHS.CommandlineArguments == RHS.CommandlineArguments) && (LHS.WorkingDirectory == RHS.WorkingDirectory) && (LHS.DataPath == RHS.DataPath);
+	}
+
+private:
+	FString Executable;
+	FString CommandlineArguments;
+	FString WorkingDirectory;
+	FString DataPath;
+	bool bShowConsole;
+
+	bool ReadFromJson(FJsonObject& JsonObject);
+	void WriteToJson(TJsonWriter<TCHAR, TPrettyJsonPrintPolicy<TCHAR>>& Writer) const;
+
+	bool ReadFromJsonFile(const TCHAR* Filename);
+	bool WriteToJsonFile(const TCHAR* Filename) const;
+
+	friend class FZenServiceInstance;
+	friend UE_API bool TryGetLocalServiceRunContext(FZenLocalServiceRunContext& OutContext);
+};
+
+UE_API bool IsLocalServiceRunning(const TCHAR* DataPath, uint16* OutPort = nullptr);
+UE_API FProcHandle StartLocalService(const FZenLocalServiceRunContext& Context, const TCHAR* TransientArgs = nullptr);
+UE_API bool StopLocalService(const TCHAR* DataPath, double MaximumWaitDurationSeconds = 5.0);
+
+UE_API FString GetLocalServiceInstallPath();
+UE_API FString GetLocalServiceInstallVersion(bool bDetailed = true);
+UE_API FString GetLocalInstallUtilityPath();
+
+struct FGCStatus
+{
+	FString Description;
+};
+
+/**
+ * Type used to declare usage of a Zen server instance whether the shared default instance or a unique non-default instance.
+ * Used to help manage launch, and optionally in the future, shutdown a shared default instance.  Use the default constructor
+ * to reference the default instance (which may be launched on demand), or use the non-default constructors with a specific
+ * URL or HostName/Port pair which is required to pre-exist (will not be auto-launched).
+ */
+class FScopeZenService
+{
+public:
+	UE_API FScopeZenService();
+	UE_API FScopeZenService(FStringView InstanceURL);
+	UE_API FScopeZenService(FServiceSettings&& InSettings);
+	UE_API ~FScopeZenService();
+
+
+	UE_API const FZenServiceInstance& GetInstance() const { return *ServiceInstance; }
+	UE_API FZenServiceInstance& GetInstance() { return *ServiceInstance; }
+
+private:
+	FZenServiceInstance* ServiceInstance;
+	TUniquePtr<UE::Zen::FZenServiceInstance> UniqueNonDefaultInstance;
+};
+
+/**
+ * Gets the default Zen service instance.  The default instance is configured through ZenServiceInstance INI section.
+ * The default instance can (depending on configuration):
+ *  - Auto-launched on demand (optionally elevated on Windows)
+ *  - Be copied out of the workspace tree before execution
+ *  - Shared between multiple tools running concurrently (implemented by launching multiple instances and expecting them to communicate and shutdown as needed)
+ *  - Instigate self-shutdown when all processes that requested it have terminated
+ *  - Use a subdirectory of the existing local DDC cache path as Zen's data path
+ *  - Use an upstream Zen service
+ *  - Be overridden by commandline arguments to reference an existing running instance instead of auto-launching a new instance.
+ * Note that no assumptions should be made about the hostname/port of the default instance.  Calling code should instead expect
+ * that the instance is authoritative over the hostname/port that it will end up using and should query that information from the
+ * instance as needed.
+ */
+UE_API FZenServiceInstance& GetDefaultServiceInstance();
+UE_API bool IsDefaultServicePresent();
+
+/**
+ * A representation of a Zen service instance.  Generally not accessed directly, but via FScopeZenService.
+ */
+class FZenServiceInstance
+{
+public:
+
+	UE_API FZenServiceInstance();
+	UE_API FZenServiceInstance(FStringView InstanceURL);
+	UE_API FZenServiceInstance(FServiceSettings&& InSettings);
+	UE_API ~FZenServiceInstance();
+
+	inline const TCHAR* GetURL() const { return *URL; }
+	inline const TCHAR* GetHostName() const { return *HostName; }
+	inline uint16 GetPort() const { return Port; }
+	inline const FServiceSettings& GetServiceSettings() const { return Settings; }
+	UE_API bool IsServiceRunning();
+	UE_API bool IsServiceReady();
+	UE_API bool IsServiceRunningLocally() const { return bIsRunningLocally; }
+
+	UE_API bool GetStats(FZenStats& Stats);
+	UE_API bool GetGCStatus(FGCStatus& Status);
+	UE_API bool RequestGC(const bool* OverrideCollectSmallObjects = nullptr, const uint32* OverrideMaxCacheDuration = nullptr);
+
+	static UE_API uint16 GetAutoLaunchedPort();
+
+private:
+
+	void Initialize();
+	void PromptUserToStopRunningServerInstance(const FString& ServerFilePath);
+	FString ConditionalUpdateLocalInstall();
+	static bool AutoLaunch(const FServiceAutoLaunchSettings& InSettings, FString&& ExecutablePath, FString& OutHostName, uint16& OutPort);
+	void OnAnalyticsEvent(TArray<FAnalyticsEventAttribute>& Attributes);
+
+	mutable TPimplPtr<class FZenHttpRequest> GCStatusHttpRequest;
+	mutable TFuture<TOptional<FGCStatus>> GCStatusRequest;
+	mutable TOptional<FGCStatus> LastGCStatus;
+	mutable uint64 LastGCStatusTime = 0;
+
+	mutable TPimplPtr<class FZenHttpRequest> StatsHttpRequest;
+	mutable TFuture<FZenStats> StatsRequest;
+	mutable FZenStats LastStats;
+	mutable uint64 LastStatsTime = 0;
+
+	FServiceSettings Settings;
+	FString URL;
+	FString HostName;
+	uint16 Port;
+	static uint16 AutoLaunchedPort;
+	bool bHasLaunchedLocal = false;
+	bool bIsRunningLocally = true;
+};
+
+} // namespace UE::Zen
+
+#endif // UE_WITH_ZEN
+
+#undef UE_API
diff --git a/Engine/Source/Developer/Zen/Public/Experimental/ZenStatistics.h b/Engine/Source/Developer/Zen/Public/Experimental/ZenStatistics.h
new file mode 100644
--- /dev/null
+++ b/Engine/Source/Developer/Zen/Public/Experimental/ZenStatistics.h
@@ -0,0 +1,84 @@
+// Copyright Epic Games, Inc. All Rights Reserved.
+
+#pragma once
+
+#include "Containers/UnrealString.h"
+#include "ZenGlobals.h"
+
+#if UE_WITH_ZEN
+
+namespace UE::Zen {
+
+	struct FZenCacheSizeStats
+	{
+		double Disk = 0;
+		double Memory = 0;
+	};
+
+	struct FZenCacheStats
+	{
+		FZenCacheSizeStats Size;
+		int64 Hits = 0;
+		int64 Misses = 0;
+		double HitRatio = 0.0;
+		int64 UpstreamHits = 0;
+		double UpstreamRatio = 0.0;	
+	};
+
+	struct FZenRequestStats
+	{
+		int64 Count = 0;
+		double RateMean = 0.0;
+		double TAverage = 0.0;
+		double TMin = 0.0;
+		double TMax = 0.0;
+	};
+
+	struct FZenEndPointStats
+	{
+		FString Name;
+		FString Url;
+		FString Health;
+		double HitRatio = 0.0;
+		double DownloadedMB = 0;
+		double UploadedMB = 0;
+		int64 ErrorCount = 0;
+	};
+
+	struct FZenUpstreamStats
+	{
+		bool Reading = false;
+		bool Writing = false;
+		int64 WorkerThreads = 0;
+		int64 QueueCount = 0;
+		double TotalUploadedMB = 0;
+		double TotalDownloadedMB = 0;
+		TArray<FZenEndPointStats> EndPointStats;
+	};
+
+	struct FZenCASSizeStats
+	{
+		int64 Tiny = 0;
+		int64 Small = 0;
+		int64 Large = 0;
+		int64 Total = 0;
+	};
+
+	struct FZenCASStats
+	{
+		FZenCASSizeStats Size;
+	};
+
+	struct FZenStats
+	{
+		FZenRequestStats RequestStats;
+		FZenCacheStats CacheStats;
+		FZenUpstreamStats UpstreamStats;
+		FZenRequestStats UpstreamRequestStats;
+		FZenCASStats CASStats;
+		bool IsValid = false;
+	};
+
+} // namespace UE::Zen
+
+#endif // UE_WITH_ZEN
diff --git a/Engine/Source/Developer/Zen/Public/ZenGlobals.h b/Engine/Source/Developer/Zen/Public/ZenGlobals.h
deleted file mode 100644
diff --git a/Engine/Source/Developer/Zen/Public/ZenServerInterface.h b/Engine/Source/Developer/Zen/Public/ZenServerInterface.h
deleted file mode 100644
diff --git a/Engine/Source/Developer/Zen/Public/ZenStatistics.h b/Engine/Source/Developer/Zen/Public/ZenStatistics.h
deleted file mode 100644
diff --git a/Engine/Source/Editor/DerivedDataEditor/Private/DerivedDataEditorModule.cpp b/Engine/Source/Editor/DerivedDataEditor/Private/DerivedDataEditorModule.cpp
--- a/Engine/Source/Editor/DerivedDataEditor/Private/DerivedDataEditorModule.cpp
+++ b/Engine/Source/Editor/DerivedDataEditor/Private/DerivedDataEditorModule.cpp
@@ -1,17 +1,17 @@
 // Copyright Epic Games, Inc. All Rights Reserved.
 
 #include "DerivedDataEditorModule.h"
-#include "DerivedDataInformation.h"
-#include "SDerivedDataStatusBar.h"
-#include "SDerivedDataDialogs.h"
-#include "SDerivedDataCacheSettings.h"
-#include "ZenServerInterface.h"
 #include "DerivedDataCacheNotifications.h"
-#include "Modules/ModuleManager.h"
+#include "DerivedDataInformation.h"
+#include "Experimental/ZenServerInterface.h"
 #include "Framework/Application/SlateApplication.h"
+#include "Modules/ModuleManager.h"
+#include "SDerivedDataCacheSettings.h"
+#include "SDerivedDataDialogs.h"
+#include "SDerivedDataStatusBar.h"
+#include "Styling/AppStyle.h"
 #include "Widgets/DeclarativeSyntaxSupport.h"
 #include "Widgets/Docking/SDockTab.h"
-#include "Styling/AppStyle.h"
 #include "WorkspaceMenuStructure.h"
 #include "WorkspaceMenuStructureModule.h"
 
diff --git a/Engine/Source/Editor/TurnkeySupport/Private/TurnkeySupportModule.cpp b/Engine/Source/Editor/TurnkeySupport/Private/TurnkeySupportModule.cpp
--- a/Engine/Source/Editor/TurnkeySupport/Private/TurnkeySupportModule.cpp
+++ b/Engine/Source/Editor/TurnkeySupport/Private/TurnkeySupportModule.cpp
@@ -48,7 +48,7 @@
 #include "ITurnkeyIOModule.h"
 #include "AnalyticsEventAttribute.h"
 #if WITH_EDITOR
-#include "ZenServerInterface.h"
+#include "Experimental/ZenServerInterface.h"
 #endif
 
 #include "Misc/App.h"
diff --git a/Engine/Source/Editor/UnrealEd/Private/DerivedDataBuildRemoteExecutor.cpp b/Engine/Source/Editor/UnrealEd/Private/DerivedDataBuildRemoteExecutor.cpp
--- a/Engine/Source/Editor/UnrealEd/Private/DerivedDataBuildRemoteExecutor.cpp
+++ b/Engine/Source/Editor/UnrealEd/Private/DerivedDataBuildRemoteExecutor.cpp
@@ -14,6 +14,7 @@
 #include "DerivedDataRequestOwner.h"
 #include "DerivedDataSharedString.h"
 #include "DerivedDataValue.h"
+#include "Experimental/ZenServerInterface.h"
 #include "Features/IModularFeatures.h"
 #include "HAL/Event.h"
 #include "HAL/PlatformProcess.h"
@@ -30,7 +31,6 @@
 #include "Tasks/Task.h"
 #include "Templates/Function.h"
 #include "ZenServerHttp.h"
-#include "ZenServerInterface.h"
 
 #include <atomic>
 
diff --git a/Engine/Source/Editor/UnrealEd/Private/PlayLevelViaLauncher.cpp b/Engine/Source/Editor/UnrealEd/Private/PlayLevelViaLauncher.cpp
--- a/Engine/Source/Editor/UnrealEd/Private/PlayLevelViaLauncher.cpp
+++ b/Engine/Source/Editor/UnrealEd/Private/PlayLevelViaLauncher.cpp
@@ -1,31 +1,31 @@
 // Copyright Epic Games, Inc. All Rights Reserved.
 
+#include "Algo/AllOf.h"
+#include "AnalyticsEventAttribute.h"
+#include "Async/Async.h"
+#include "CookerSettings.h"
+#include "DesktopPlatformModule.h"
 #include "Editor/EditorEngine.h"
-#include "ITargetDeviceServicesModule.h"
-#include "ILauncherServicesModule.h"
+#include "Editor/EditorPerProjectUserSettings.h"
 #include "EditorAnalytics.h"
-#include "AnalyticsEventAttribute.h"
-#include "Widgets/Notifications/SNotificationList.h"
-#include "Interfaces/ITargetPlatform.h"
-#include "Misc/CoreMisc.h"
+#include "Experimental/ZenServerInterface.h"
+#include "Framework/Docking/TabManager.h"
+#include "Framework/Notifications/NotificationManager.h"
 #include "GameProjectGenerationModule.h"
-#include "CookerSettings.h"
-#include "Settings/EditorExperimentalSettings.h"
-#include "UnrealEdMisc.h"
+#include "ILauncherServicesModule.h"
+#include "Interfaces/ITargetPlatform.h"
 #include "Interfaces/ITargetPlatformManagerModule.h"
-#include "Settings/ProjectPackagingSettings.h"
-#include "Framework/Notifications/NotificationManager.h"
-#include "PlayLevel.h"
-#include "Algo/AllOf.h"
-#include "Async/Async.h"
+#include "ITargetDeviceServicesModule.h"
 #include "Logging/MessageLog.h"
-#include "TargetReceipt.h"
-#include "DesktopPlatformModule.h"
+#include "Misc/CoreMisc.h"
 #include "PlatformInfo.h"
+#include "PlayLevel.h"
+#include "Settings/EditorExperimentalSettings.h"
 #include "Settings/LevelEditorPlaySettings.h"
-#include "Framework/Docking/TabManager.h"
-#include "Editor/EditorPerProjectUserSettings.h"
-#include "ZenServerInterface.h"
+#include "Settings/ProjectPackagingSettings.h"
+#include "TargetReceipt.h"
+#include "UnrealEdMisc.h"
+#include "Widgets/Notifications/SNotificationList.h"
 
 #define LOCTEXT_NAMESPACE "PlayLevel"
 
diff --git a/Engine/Source/Programs/ZenDashboard/Private/SZenCacheStatistics.cpp b/Engine/Source/Programs/ZenDashboard/Private/SZenCacheStatistics.cpp
--- a/Engine/Source/Programs/ZenDashboard/Private/SZenCacheStatistics.cpp
+++ b/Engine/Source/Programs/ZenDashboard/Private/SZenCacheStatistics.cpp
@@ -1,15 +1,15 @@
 // Copyright Epic Games, Inc. All Rights Reserved.
 
 #include "SZenCacheStatistics.h"
-#include "ZenServerInterface.h"
-#include "Math/UnitConversion.h"
+#include "Experimental/ZenServerInterface.h"
+#include "Internationalization/FastDecimalFormat.h"
 #include "Math/BasicMathExpressionEvaluator.h"
+#include "Math/UnitConversion.h"
 #include "Misc/ExpressionParser.h"
 #include "Styling/StyleColors.h"
 #include "Widgets/Layout/SGridPanel.h"
-#include "Widgets/Text/STextBlock.h"
 #include "Widgets/SBoxPanel.h"
-#include "Internationalization/FastDecimalFormat.h"
+#include "Widgets/Text/STextBlock.h"
 
 #define LOCTEXT_NAMESPACE "ZenDashboard"
 
diff --git a/Engine/Source/Programs/ZenDashboard/Private/SZenCacheStatistics.h b/Engine/Source/Programs/ZenDashboard/Private/SZenCacheStatistics.h
--- a/Engine/Source/Programs/ZenDashboard/Private/SZenCacheStatistics.h
+++ b/Engine/Source/Programs/ZenDashboard/Private/SZenCacheStatistics.h
@@ -2,13 +2,13 @@
 
 #pragma once
 
+#include "Experimental/ZenServerInterface.h"
 #include "ServiceInstanceManager.h"
 #include "Templates/SharedPointer.h"
 #include "Widgets/DeclarativeSyntaxSupport.h"
 #include "Widgets/SBoxPanel.h"
 #include "Widgets/SCompoundWidget.h"
 #include "Widgets/SWidget.h"
-#include "ZenServerInterface.h"
 
 class SZenCacheStatistics : public SCompoundWidget
 {
diff --git a/Engine/Source/Programs/ZenDashboard/Private/SZenServiceStatus.cpp b/Engine/Source/Programs/ZenDashboard/Private/SZenServiceStatus.cpp
--- a/Engine/Source/Programs/ZenDashboard/Private/SZenServiceStatus.cpp
+++ b/Engine/Source/Programs/ZenDashboard/Private/SZenServiceStatus.cpp
@@ -1,17 +1,17 @@
 // Copyright Epic Games, Inc. All Rights Reserved.
 
 #include "SZenServiceStatus.h"
-#include "ZenServerInterface.h"
-#include "Math/UnitConversion.h"
+#include "Experimental/ZenServerInterface.h"
+#include "Internationalization/FastDecimalFormat.h"
 #include "Math/BasicMathExpressionEvaluator.h"
+#include "Math/UnitConversion.h"
 #include "Misc/ExpressionParser.h"
 #include "Styling/StyleColors.h"
+#include "Widgets/Images/SImage.h"
+#include "Widgets/Input/SButton.h"
 #include "Widgets/Layout/SGridPanel.h"
-#include "Widgets/Text/STextBlock.h"
 #include "Widgets/SBoxPanel.h"
-#include "Internationalization/FastDecimalFormat.h"
-#include "Widgets/Input/SButton.h"
-#include "Widgets/Images/SImage.h"
+#include "Widgets/Text/STextBlock.h"
 
 #define LOCTEXT_NAMESPACE "ZenDashboard"
 
diff --git a/Engine/Source/Programs/ZenDashboard/Private/SZenServiceStatus.h b/Engine/Source/Programs/ZenDashboard/Private/SZenServiceStatus.h
--- a/Engine/Source/Programs/ZenDashboard/Private/SZenServiceStatus.h
+++ b/Engine/Source/Programs/ZenDashboard/Private/SZenServiceStatus.h
@@ -2,13 +2,13 @@
 
 #pragma once
 
+#include "Experimental/ZenServerInterface.h"
 #include "ServiceInstanceManager.h"
 #include "Templates/SharedPointer.h"
 #include "Widgets/DeclarativeSyntaxSupport.h"
 #include "Widgets/SBoxPanel.h"
 #include "Widgets/SCompoundWidget.h"
 #include "Widgets/SWidget.h"
-#include "ZenServerInterface.h"
 
 class SZenServiceStatus : public SCompoundWidget
 {
diff --git a/Engine/Source/Programs/ZenDashboard/Private/ServiceInstanceManager.cpp b/Engine/Source/Programs/ZenDashboard/Private/ServiceInstanceManager.cpp
--- a/Engine/Source/Programs/ZenDashboard/Private/ServiceInstanceManager.cpp
+++ b/Engine/Source/Programs/ZenDashboard/Private/ServiceInstanceManager.cpp
@@ -1,7 +1,7 @@
 // Copyright Epic Games, Inc. All Rights Reserved.
 
 #include "ServiceInstanceManager.h"
-#include "ZenServerInterface.h"
+#include "Experimental/ZenServerInterface.h"
 
 
 namespace UE::Zen
diff --git a/Engine/Source/Programs/ZenDashboard/Private/ZenDashboard.cpp b/Engine/Source/Programs/ZenDashboard/Private/ZenDashboard.cpp
--- a/Engine/Source/Programs/ZenDashboard/Private/ZenDashboard.cpp
+++ b/Engine/Source/Programs/ZenDashboard/Private/ZenDashboard.cpp
@@ -1,6 +1,7 @@
 // Copyright Epic Games, Inc. All Rights Reserved.
 
 #include "ZenDashboard.h"
+#include "Experimental/ZenServerInterface.h"
 #include "Features/IModularFeatures.h"
 #include "Framework/Application/SlateApplication.h"
 #include "Framework/MultiBox/MultiBoxBuilder.h"
@@ -25,7 +26,6 @@
 #include "Widgets/Notifications/SNotificationList.h"
 #include "Widgets/Text/STextBlock.h"
 #include "ZenDashboardStyle.h"
-#include "ZenServerInterface.h"
 
 #if PLATFORM_WINDOWS
 #include "Windows/WindowsHWrapper.h"
