Add capability for things that depend on the Zen service to attempt to recover the service if there is an interruption.  The Zen DDC cache store is using this to allow recovery and retry of a cache request that fails to connect or times out when communicating with a Zen service on the same machine.

diff --git a/Engine/Source/Developer/DerivedDataCache/Private/Http/CurlHttpClient.cpp b/Engine/Source/Developer/DerivedDataCache/Private/Http/CurlHttpClient.cpp
--- a/Engine/Source/Developer/DerivedDataCache/Private/Http/CurlHttpClient.cpp
+++ b/Engine/Source/Developer/DerivedDataCache/Private/Http/CurlHttpClient.cpp
@@ -168,7 +168,7 @@ private:
 	void Destroy() final { delete this; }
 
 	void ThreadLoop();
-	void CompleteRequest(CURL* Curl);
+	void CompleteRequest(CURL* Curl, CURLcode* OptionalReturnCode);
 
 	static void CurlLock(CURL* Curl, curl_lock_data Data, curl_lock_access Access, void* Param);
 	static void CurlUnlock(CURL* Curl, curl_lock_data Data, void* Param);
@@ -489,7 +489,7 @@ void FCurlHttpConnectionPool::ThreadLoop()
 					AssertMultiCodeOk(curl_multi_add_handle(CurlMulti, Curl));
 					break;
 				case EThreadCommandType::Cancel:
-					CompleteRequest(Curl);
+					CompleteRequest(Curl, nullptr);
 					break;
 				}
 			}
@@ -502,7 +502,7 @@ void FCurlHttpConnectionPool::ThreadLoop()
 		{
 			if (Message->msg == CURLMSG_DONE)
 			{
-				CompleteRequest(Message->easy_handle);
+				CompleteRequest(Message->easy_handle, &Message->data.result);
 			}
 		}
 
@@ -511,12 +511,12 @@ void FCurlHttpConnectionPool::ThreadLoop()
 	}
 }
 
-void FCurlHttpConnectionPool::CompleteRequest(CURL* Curl)
+void FCurlHttpConnectionPool::CompleteRequest(CURL* Curl, CURLcode* OptionalReturnCode)
 {
 	AssertMultiCodeOk(curl_multi_remove_handle(CurlMulti, Curl));
 	void* Request = nullptr;
 	curl_easy_getinfo(Curl, CURLINFO_PRIVATE, &Request);
-	((FCurlHttpRequest*)Request)->OnComplete(CURLE_OK);
+	((FCurlHttpRequest*)Request)->OnComplete(OptionalReturnCode ? *OptionalReturnCode : CURLE_OK);
 }
 
 void FCurlHttpConnectionPool::CurlLock(CURL* Curl, curl_lock_data Data, curl_lock_access Access, void* Param)
diff --git a/Engine/Source/Developer/DerivedDataCache/Private/ZenCacheStore.cpp b/Engine/Source/Developer/DerivedDataCache/Private/ZenCacheStore.cpp
--- a/Engine/Source/Developer/DerivedDataCache/Private/ZenCacheStore.cpp
+++ b/Engine/Source/Developer/DerivedDataCache/Private/ZenCacheStore.cpp
@@ -952,9 +952,11 @@ public:
 	FAsyncCbPackageReceiver(
 		THttpUniquePtr<IHttpRequest>&& InRequest,
 		IRequestOwner* InOwner,
+		Zen::FZenServiceInstance& InZenServiceInstance,
 		FOnRpcComplete&& InOnRpcComplete)
 		: Request(MoveTemp(InRequest))
 		, Owner(InOwner)
+		, ZenServiceInstance(InZenServiceInstance)
 		, BaseReceiver(Package, this)
 		, OnRpcComplete(MoveTemp(InOnRpcComplete))
 	{
@@ -977,11 +979,34 @@ private:
 		return &BaseReceiver;
 	}
 
+	bool ShouldRecoverAndRetry(IHttpResponse& LocalResponse)
+	{
+		if (!ZenServiceInstance.IsServiceRunningLocally())
+		{
+			return false;
+		}
+
+		if ((LocalResponse.GetErrorCode() == EHttpErrorCode::Connect) ||
+			(LocalResponse.GetErrorCode() == EHttpErrorCode::TlsConnect) ||
+			(LocalResponse.GetErrorCode() == EHttpErrorCode::TimedOut))
+		{
+			return true;
+		}
+
+		return false;
+	}
+
 	IHttpReceiver* OnComplete(IHttpResponse& LocalResponse) final
 	{
-		Request.Reset();
-		Owner->End(this, [Self = this]
+		Owner->End(this, [Self = this, &LocalResponse]
 		{
+			if (Self->ShouldRecoverAndRetry(LocalResponse) && Self->ZenServiceInstance.TryRecovery())
+			{
+				new FAsyncCbPackageReceiver(MoveTemp(Self->Request), Self->Owner, Self->ZenServiceInstance, MoveTemp(Self->OnRpcComplete));
+				return;
+			}
+
+			Self->Request.Reset();
 			if (Self->OnRpcComplete)
 			{
 				// Launch a task for the completion function since it can execute arbitrary code.
@@ -999,6 +1024,7 @@ private:
 	THttpUniquePtr<IHttpResponse> Response;
 	TRefCountPtr<IHttpResponseMonitor> Monitor;
 	IRequestOwner* Owner;
+	Zen::FZenServiceInstance& ZenServiceInstance;
 	FCbPackage Package;
 	FCbPackageReceiver BaseReceiver;
 	FOnRpcComplete OnRpcComplete;
@@ -1068,7 +1094,7 @@ void FZenCacheStore::EnqueueAsyncRpc(IRequestOwner& Owner, FCbObject RequestObje
 	THttpUniquePtr<IHttpRequest> Request = CreateRpcRequest();
 	Request->SetContentType(EHttpMediaType::CbObject);
 	Request->SetBody(RequestObject.GetBuffer().MakeOwned());
-	new FAsyncCbPackageReceiver(MoveTemp(Request), &Owner, MoveTemp(OnComplete));
+	new FAsyncCbPackageReceiver(MoveTemp(Request), &Owner, ZenService.GetInstance(), MoveTemp(OnComplete));
 }
 
 void FZenCacheStore::EnqueueAsyncRpc(IRequestOwner& Owner, const FCbPackage& RequestPackage, FOnRpcComplete&& OnComplete)
@@ -1076,7 +1102,7 @@ void FZenCacheStore::EnqueueAsyncRpc(IRequestOwner& Owner, const FCbPackage& Req
 	THttpUniquePtr<IHttpRequest> Request = CreateRpcRequest();
 	Request->SetContentType(EHttpMediaType::CbPackage);
 	Request->SetBody(SaveRpcPackage(RequestPackage));
-	new FAsyncCbPackageReceiver(MoveTemp(Request), &Owner, MoveTemp(OnComplete));
+	new FAsyncCbPackageReceiver(MoveTemp(Request), &Owner, ZenService.GetInstance(), MoveTemp(OnComplete));
 }
 
 void FZenCacheStore::LegacyStats(FDerivedDataCacheStatsNode& OutNode)
diff --git a/Engine/Source/Developer/Zen/Private/ZenServerInterface.cpp b/Engine/Source/Developer/Zen/Private/ZenServerInterface.cpp
--- a/Engine/Source/Developer/Zen/Private/ZenServerInterface.cpp
+++ b/Engine/Source/Developer/Zen/Private/ZenServerInterface.cpp
@@ -1103,17 +1103,83 @@ FZenServiceInstance::IsServiceReady()
 		
 		if (Result == Zen::FZenHttpRequest::Result::Success && Zen::IsSuccessCode(Request.GetResponseCode()))
 		{
-			UE_LOG(LogZenServiceInstance, Display, TEXT("Z$ HTTP DDC service status: %s."), *Request.GetResponseAsString());
+			UE_LOG(LogZenServiceInstance, Display, TEXT("ZenServer HTTP service status: %s."), *Request.GetResponseAsString());
 			return true;
 		}
 		else
 		{
-			UE_LOG(LogZenServiceInstance, Warning, TEXT("Unable to reach Z$ HTTP DDC service at %s. Status: %d . Response: %s"), ZenDomain.ToString(), Request.GetResponseCode(), *Request.GetResponseAsString());
+			UE_LOG(LogZenServiceInstance, Warning, TEXT("Unable to reach ZenServer HTTP service at %s. Status: %d . Response: %s"), ZenDomain.ToString(), Request.GetResponseCode(), *Request.GetResponseAsString());
 		}
 	}
 	return false;
 }
 
+bool 
+FZenServiceInstance::TryRecovery()
+{
+	if (!bHasLaunchedLocal)
+	{
+		return false;
+	}
+
+	static FCriticalSection RecoveryCriticalSection;
+	static std::atomic<int64> LastRecoveryTicks;
+	static bool bLastRecoveryResult = false;
+	const FTimespan MaximumWaitForHealth = FTimespan::FromSeconds(30);
+	const FTimespan MinimumDurationSinceLastRecovery = FTimespan::FromMinutes(2);
+
+	FTimespan TimespanSinceLastRecovery = FDateTime::UtcNow() - FDateTime(LastRecoveryTicks.load(std::memory_order_relaxed));
+
+	if (TimespanSinceLastRecovery > MinimumDurationSinceLastRecovery)
+	{
+		FScopeLock Lock(&RecoveryCriticalSection);
+		// Update timespan since it may have changed since we waited to enter the crit section
+		TimespanSinceLastRecovery = FDateTime::UtcNow() - FDateTime(LastRecoveryTicks.load(std::memory_order_relaxed));
+		if (TimespanSinceLastRecovery > MinimumDurationSinceLastRecovery)
+		{
+			UE_LOG(LogZenServiceInstance, Display, TEXT("Local ZenServer recovery being attempted..."));
+			FZenLocalServiceRunContext RunContext;
+			if (TryGetLocalServiceRunContext(RunContext))
+			{
+				StopLocalService(*RunContext.GetDataPath());
+				StartLocalService(RunContext);
+				UE_LOG(LogZenServiceInstance, Display, TEXT("Local ZenServer recovery finished."));
+			}
+			else
+			{
+				UE_LOG(LogZenServiceInstance, Warning, TEXT("Local ZenServer recovery failed due to lack of run context."));
+			}
+			
+			FDateTime StartedWaitingForHealth = FDateTime::UtcNow();
+			bLastRecoveryResult = IsServiceReady();
+			while (!bLastRecoveryResult)
+			{
+				FTimespan WaitForHealth = StartedWaitingForHealth - FDateTime::UtcNow();
+				if (WaitForHealth > MaximumWaitForHealth)
+				{
+					UE_LOG(LogZenServiceInstance, Warning, TEXT("Local ZenServer recovery timed out waiting for service to become healthy"));
+					break;
+				}
+
+				FPlatformProcess::Sleep(0.5f);
+				bLastRecoveryResult = IsServiceReady();
+			}
+			LastRecoveryTicks.store(FDateTime::UtcNow().GetTicks(), std::memory_order_relaxed);
+			
+			if (bLastRecoveryResult)
+			{
+				UE_LOG(LogZenServiceInstance, Display, TEXT("Local ZenServer post recovery status: Healthy"));
+			}
+			else
+			{
+				UE_LOG(LogZenServiceInstance, Display, TEXT("Local ZenServer post recovery status: NOT healthy"));
+			}
+		}
+	}
+
+	return bLastRecoveryResult;
+}
+
 uint16
 FZenServiceInstance::GetAutoLaunchedPort()
 {
diff --git a/Engine/Source/Developer/Zen/Public/Experimental/ZenServerInterface.h b/Engine/Source/Developer/Zen/Public/Experimental/ZenServerInterface.h
--- a/Engine/Source/Developer/Zen/Public/Experimental/ZenServerInterface.h
+++ b/Engine/Source/Developer/Zen/Public/Experimental/ZenServerInterface.h
@@ -182,6 +182,8 @@ public:
 	UE_API bool IsServiceReady();
 	UE_API bool IsServiceRunningLocally() const { return bIsRunningLocally; }
 
+	UE_API bool TryRecovery();
+
 	UE_API bool GetStats(FZenStats& Stats);
 	UE_API bool GetGCStatus(FGCStatus& Status);
 	UE_API bool RequestGC(const bool* OverrideCollectSmallObjects = nullptr, const uint32* OverrideMaxCacheDuration = nullptr);
