Missing changes from robomerge of CL 23582398.
ZenServerInterface.cpp and ZenServerInterface.h are moved files and need to be updated in their new location.
#rnx

diff --git a/Engine/Source/Developer/Zen/Private/ZenServerInterface.cpp b/Engine/Source/Developer/Zen/Private/ZenServerInterface.cpp
--- a/Engine/Source/Developer/Zen/Private/ZenServerInterface.cpp
+++ b/Engine/Source/Developer/Zen/Private/ZenServerInterface.cpp
@@ -6,8 +6,6 @@
 #include "ZenSerialization.h"
 #include "ZenServerHttp.h"
 
-#include "Analytics.h"
-#include "AnalyticsEventAttribute.h"
 #include "Async/Async.h"
 #include "Dom/JsonValue.h"
 #include "HAL/FileManager.h"
@@ -1081,10 +1079,6 @@ FZenServiceInstance::FZenServiceInstance(FServiceSettings&& InSettings)
 
 FZenServiceInstance::~FZenServiceInstance()
 {
-	if (FAnalytics::IsAvailable())
-	{
-		FAnalytics::Get().GetEventCallback(TEXT("Core.Loading"))->RemoveAll(this);
-	}
 }
 
 bool 
@@ -1208,8 +1202,6 @@ FZenServiceInstance::Initialize()
 		bIsRunningLocally = IsLocalHost(HostName);
 	}
 	URL = WriteToString<64>(TEXT("http://"), HostName, TEXT(":"), Port, TEXT("/"));
-
-	FAnalytics::Get().GetEventCallback(TEXT("Core.Loading"))->AddRaw(this,&FZenServiceInstance::OnAnalyticsEvent);
 }
 
 void
@@ -1605,191 +1597,6 @@ FZenServiceInstance::GetStats(FZenStats& Stats)
 	return Stats.IsValid;
 }
 
-bool 
-FZenServiceInstance::GetGCStatus(FGCStatus& Status)
-{
-	check(IsInGameThread());
-
-	// If we've already requested status and it is ready then grab it
-	if (GCStatusRequest.IsReady() == true )
-	{
-		LastGCStatus	 = GCStatusRequest.Get();
-		LastGCStatusTime = FPlatformTime::Cycles64();
-
-		GCStatusRequest.Reset();
-	}
-	
-	// Make a copy of the last updated status
-	if (LastGCStatus.IsSet())
-	{
-		Status = LastGCStatus.GetValue();
-	}
-
-	const uint64 CurrentTime = FPlatformTime::Cycles64();
-	constexpr double MinTimeBetweenRequestsInSeconds = 0.5;
-	const double DeltaTimeInSeconds = FPlatformTime::ToSeconds64(CurrentTime - LastGCStatusTime);
-
-	if (!GCStatusRequest.IsValid() && DeltaTimeInSeconds > MinTimeBetweenRequestsInSeconds)
-	{
-#if WITH_EDITOR
-		EAsyncExecution ThreadPool = EAsyncExecution::LargeThreadPool;
-#else
-		EAsyncExecution ThreadPool = EAsyncExecution::ThreadPool;
-#endif
-		if (!GCStatusHttpRequest.IsValid())
-		{
-			TStringBuilder<128> ZenDomain;
-			ZenDomain << HostName << TEXT(":") << Port;
-			GCStatusHttpRequest = MakePimpl<FZenHttpRequest>(ZenDomain.ToString(), false);
-		}
-
-		// We've not got any requests in flight and we've met a given time requirement for requests
-		GCStatusRequest = Async(ThreadPool, [this]
-			{
-				UE::Zen::FZenHttpRequest& Request = *GCStatusHttpRequest.Get();
-				Request.Reset();
-
-				TArray64<uint8> GetBuffer;
-				FZenHttpRequest::Result Result = Request.PerformBlockingDownload(TEXTVIEW("/admin/gc"), &GetBuffer, Zen::EContentType::CbObject);
-
-				TOptional<FGCStatus> GCStatus;
-
-				if (Result == Zen::FZenHttpRequest::Result::Success && Request.GetResponseCode() == 200)
-				{
-					FCbObjectView RootObjectView(GetBuffer.GetData());
-
-					GCStatus.Emplace();
-					GCStatus->Description = FString(RootObjectView["Status"].AsString());
-				}
-
-				return GCStatus;
-			});
-	}
-
-	return LastGCStatus.IsSet();
-}
-
-bool 
-FZenServiceInstance::RequestGC(const bool* OverrideCollectSmallObjects, const uint32* OverrideMaxCacheDuration)
-{
-	TStringBuilder<128> ZenDomain;
-	ZenDomain << HostName << TEXT(":") << Port;
-	UE::Zen::FZenHttpRequest Request(ZenDomain.ToString(), false);
-
-	TCHAR Separators[] = {TEXT('?'), TEXT('&')};
-	int32 SeparatorIndex = 0;
-	TStringBuilder<128> Query;
-	Query << TEXTVIEW("/admin/gc");
-
-	if (OverrideCollectSmallObjects)
-	{
-		Query << Separators[SeparatorIndex] << LexToString(*OverrideCollectSmallObjects);
-		SeparatorIndex = FMath::Min(SeparatorIndex + 1, (int32)UE_ARRAY_COUNT(Separators));
-	}
-
-	if (OverrideMaxCacheDuration)
-	{
-		Query << Separators[SeparatorIndex] << LexToString(*OverrideMaxCacheDuration);
-		SeparatorIndex = FMath::Min(SeparatorIndex + 1, (int32)UE_ARRAY_COUNT(Separators));
-	}
-
-	FZenHttpRequest::Result Result = Request.PerformBlockingPost(Query.ToString(), FMemoryView());
-
-	if (Result == Zen::FZenHttpRequest::Result::Success && Request.GetResponseCode() == 200)
-	{
-		FCbObjectView ResponseObject = FCbObjectView(Request.GetResponseBuffer().GetData());
-		FUtf8StringView ResponseStatus = ResponseObject["status"].AsString();
-
-		return (ResponseStatus == "Started") || (ResponseStatus == "Running");
-	}
-	return false;
-}
-
-void FZenServiceInstance::OnAnalyticsEvent(TArray<FAnalyticsEventAttribute>& Attributes )
-{
-	// Grab the Zen summary stats
-	FZenStats ZenStats;
-
-	GetStats(ZenStats);
-
-	const FString BaseName = TEXT("Zen");
-
-	{
-		FString AttrName = BaseName + TEXT(".Enabled");
-		Attributes.Emplace(MoveTemp(AttrName), ZenStats.IsValid);
-	}
-
-	{
-		FString AttrName = BaseName + TEXT(".Cache.HitRatio");
-		Attributes.Emplace(MoveTemp(AttrName), ZenStats.CacheStats.HitRatio);
-	}
-
-	{
-		FString AttrName = BaseName + TEXT(".Cache.Hits");
-		Attributes.Emplace(MoveTemp(AttrName), ZenStats.CacheStats.Hits);
-	}
-
-	{
-		FString AttrName = BaseName + TEXT(".Cache.Misses");
-		Attributes.Emplace(MoveTemp(AttrName), ZenStats.CacheStats.Misses);
-	}
-
-	{
-		FString AttrName = BaseName + TEXT(".Cache.Size.Disk");
-		Attributes.Emplace(MoveTemp(AttrName), ZenStats.CacheStats.Size.Disk);
-	}
-
-	{
-		FString AttrName = BaseName + TEXT(".Cache.Size.Memory");
-		Attributes.Emplace(MoveTemp(AttrName), ZenStats.CacheStats.Size.Memory);
-	}
-
-	{
-		FString AttrName = BaseName + TEXT(".Cache.UpstreamHits");
-		Attributes.Emplace(MoveTemp(AttrName), ZenStats.CacheStats.UpstreamHits);
-	}
-
-	{
-		FString AttrName = BaseName + TEXT(".Cache.UpstreamRatio");
-		Attributes.Emplace(MoveTemp(AttrName), ZenStats.CacheStats.UpstreamRatio);
-	}
-
-	{
-		FString AttrName = BaseName + TEXT(".Cache.TotalUploadedMB");
-		Attributes.Emplace(MoveTemp(AttrName), ZenStats.UpstreamStats.TotalUploadedMB);
-	}
-
-	{
-		FString AttrName = BaseName + TEXT(".Upstream.TotalDownloadedMB");
-		Attributes.Emplace(MoveTemp(AttrName), ZenStats.UpstreamStats.TotalDownloadedMB);
-	}
-
-	{
-		FString AttrName = BaseName + TEXT(".Upstream.TotalUploadedMB");
-		Attributes.Emplace(MoveTemp(AttrName), ZenStats.UpstreamStats.TotalUploadedMB);
-	}
-
-	{
-		FString AttrName = BaseName + TEXT(".Cas.Size.Large");
-		Attributes.Emplace(MoveTemp(AttrName), ZenStats.CASStats.Size.Large);
-	}
-
-	{
-		FString AttrName = BaseName + TEXT(".Cas.Size.Small");
-		Attributes.Emplace(MoveTemp(AttrName), ZenStats.CASStats.Size.Small);
-	}
-
-	{
-		FString AttrName = BaseName + TEXT(".Cas.Size.Tiny");
-		Attributes.Emplace(MoveTemp(AttrName), ZenStats.CASStats.Size.Tiny);
-	}
-
-	{
-		FString AttrName = BaseName + TEXT(".Cas.Size.Total");
-		Attributes.Emplace(MoveTemp(AttrName), ZenStats.CASStats.Size.Total);
-	}
-}
-
 #endif // UE_WITH_ZEN
 
 }
diff --git a/Engine/Source/Developer/Zen/Public/Experimental/ZenServerInterface.h b/Engine/Source/Developer/Zen/Public/Experimental/ZenServerInterface.h
--- a/Engine/Source/Developer/Zen/Public/Experimental/ZenServerInterface.h
+++ b/Engine/Source/Developer/Zen/Public/Experimental/ZenServerInterface.h
@@ -23,8 +23,6 @@
 
 #define UE_API ZEN_API
 
-struct FAnalyticsEventAttribute;
-
 namespace UE::Zen
 {
 
@@ -196,7 +194,6 @@ private:
 	void PromptUserToStopRunningServerInstance(const FString& ServerFilePath);
 	FString ConditionalUpdateLocalInstall();
 	static bool AutoLaunch(const FServiceAutoLaunchSettings& InSettings, FString&& ExecutablePath, FString& OutHostName, uint16& OutPort);
-	void OnAnalyticsEvent(TArray<FAnalyticsEventAttribute>& Attributes);
 
 	mutable TPimplPtr<class FZenHttpRequest> GCStatusHttpRequest;
 	mutable TFuture<TOptional<FGCStatus>> GCStatusRequest;
