Fix for Zen GC methods removed when merging.

diff --git a/Engine/Source/Developer/Zen/Private/ZenServerInterface.cpp b/Engine/Source/Developer/Zen/Private/ZenServerInterface.cpp
--- a/Engine/Source/Developer/Zen/Private/ZenServerInterface.cpp
+++ b/Engine/Source/Developer/Zen/Private/ZenServerInterface.cpp
@@ -1597,6 +1597,106 @@ FZenServiceInstance::GetStats(FZenStats& Stats)
 	return Stats.IsValid;
 }
 
+bool 
+FZenServiceInstance::GetGCStatus(FGCStatus& Status)
+{
+	check(IsInGameThread());
+
+	// If we've already requested status and it is ready then grab it
+	if (GCStatusRequest.IsReady() == true )
+	{
+		LastGCStatus	 = GCStatusRequest.Get();
+		LastGCStatusTime = FPlatformTime::Cycles64();
+
+		GCStatusRequest.Reset();
+	}
+	
+	// Make a copy of the last updated status
+	if (LastGCStatus.IsSet())
+	{
+		Status = LastGCStatus.GetValue();
+	}
+
+	const uint64 CurrentTime = FPlatformTime::Cycles64();
+	constexpr double MinTimeBetweenRequestsInSeconds = 0.5;
+	const double DeltaTimeInSeconds = FPlatformTime::ToSeconds64(CurrentTime - LastGCStatusTime);
+
+	if (!GCStatusRequest.IsValid() && DeltaTimeInSeconds > MinTimeBetweenRequestsInSeconds)
+	{
+#if WITH_EDITOR
+		EAsyncExecution ThreadPool = EAsyncExecution::LargeThreadPool;
+#else
+		EAsyncExecution ThreadPool = EAsyncExecution::ThreadPool;
+#endif
+		if (!GCStatusHttpRequest.IsValid())
+		{
+			TStringBuilder<128> ZenDomain;
+			ZenDomain << HostName << TEXT(":") << Port;
+			GCStatusHttpRequest = MakePimpl<FZenHttpRequest>(ZenDomain.ToString(), false);
+		}
+
+		// We've not got any requests in flight and we've met a given time requirement for requests
+		GCStatusRequest = Async(ThreadPool, [this]
+			{
+				UE::Zen::FZenHttpRequest& Request = *GCStatusHttpRequest.Get();
+				Request.Reset();
+
+				TArray64<uint8> GetBuffer;
+				FZenHttpRequest::Result Result = Request.PerformBlockingDownload(TEXTVIEW("/admin/gc"), &GetBuffer, Zen::EContentType::CbObject);
+
+				TOptional<FGCStatus> GCStatus;
+
+				if (Result == Zen::FZenHttpRequest::Result::Success && Request.GetResponseCode() == 200)
+				{
+					FCbObjectView RootObjectView(GetBuffer.GetData());
+
+					GCStatus.Emplace();
+					GCStatus->Description = FString(RootObjectView["Status"].AsString());
+				}
+
+				return GCStatus;
+			});
+	}
+
+	return LastGCStatus.IsSet();
+}
+
+bool 
+FZenServiceInstance::RequestGC(const bool* OverrideCollectSmallObjects, const uint32* OverrideMaxCacheDuration)
+{
+	TStringBuilder<128> ZenDomain;
+	ZenDomain << HostName << TEXT(":") << Port;
+	UE::Zen::FZenHttpRequest Request(ZenDomain.ToString(), false);
+
+	TCHAR Separators[] = {TEXT('?'), TEXT('&')};
+	int32 SeparatorIndex = 0;
+	TStringBuilder<128> Query;
+	Query << TEXTVIEW("/admin/gc");
+
+	if (OverrideCollectSmallObjects)
+	{
+		Query << Separators[SeparatorIndex] << LexToString(*OverrideCollectSmallObjects);
+		SeparatorIndex = FMath::Min(SeparatorIndex + 1, (int32)UE_ARRAY_COUNT(Separators));
+	}
+
+	if (OverrideMaxCacheDuration)
+	{
+		Query << Separators[SeparatorIndex] << LexToString(*OverrideMaxCacheDuration);
+		SeparatorIndex = FMath::Min(SeparatorIndex + 1, (int32)UE_ARRAY_COUNT(Separators));
+	}
+
+	FZenHttpRequest::Result Result = Request.PerformBlockingPost(Query.ToString(), FMemoryView());
+
+	if (Result == Zen::FZenHttpRequest::Result::Success && Request.GetResponseCode() == 200)
+	{
+		FCbObjectView ResponseObject = FCbObjectView(Request.GetResponseBuffer().GetData());
+		FUtf8StringView ResponseStatus = ResponseObject["status"].AsString();
+
+		return (ResponseStatus == "Started") || (ResponseStatus == "Running");
+	}
+	return false;
+}
+
 #endif // UE_WITH_ZEN
 
 }
