Output zen command utility output to log if it returns an error code

diff --git a/Engine/Source/Developer/Zen/Private/ZenServerInterface.cpp b/Engine/Source/Developer/Zen/Private/ZenServerInterface.cpp
--- a/Engine/Source/Developer/Zen/Private/ZenServerInterface.cpp
+++ b/Engine/Source/Developer/Zen/Private/ZenServerInterface.cpp
@@ -101,16 +101,17 @@ GetZenVersion(const FString& UtilityPath, const FString& ServicePath, const FStr
 			}
 		}
 
+		FString OutputString = MonitoredUtilityProcess.GetFullOutputWithoutDelegate();
 		if (MonitoredUtilityProcess.GetReturnCode() != 0)
 		{
-			checkf(false, TEXT("Unexpected return code after launch of zen utility for gathering version data: '%s' (%d)."), *UtilityPath, MonitoredUtilityProcess.GetReturnCode());
+			checkf(false, TEXT("Unexpected return code after launch of zen utility for gathering version data: '%s' (%d). Output: '%s'"), *UtilityPath, MonitoredUtilityProcess.GetReturnCode(), *OutputString);
 			return GetFallbackVersion();
 		}
 
-		FString VersionOutputString = MonitoredUtilityProcess.GetFullOutputWithoutDelegate().TrimStartAndEnd();
+		FString VersionOutputString = OutputString.TrimStartAndEnd();
 		if (!ComparableVersion.TryParse(*VersionOutputString))
 		{
-			checkf(false, TEXT("Invalid version information after launch of zen utility for gathering version data: '%s' (%s)."), *UtilityPath, *VersionOutputString);
+			checkf(false, TEXT("Invalid version information after launch of zen utility for gathering version data: '%s' (`%s`)"), *UtilityPath, *VersionOutputString);
 			return GetFallbackVersion();
 		}
 
