Fix for escaped character handling when a Local DDC path is stored with escaped characters in the EditorSettings.

diff --git a/Engine/Source/Developer/Zen/Private/ZenServerInterface.cpp b/Engine/Source/Developer/Zen/Private/ZenServerInterface.cpp
--- a/Engine/Source/Developer/Zen/Private/ZenServerInterface.cpp
+++ b/Engine/Source/Developer/Zen/Private/ZenServerInterface.cpp
@@ -210,7 +210,8 @@ DetermineLocalDataCachePath(const TCHAR* ConfigSection, FString& DataPath)
 			FString SettingPath;
 			if (FParse::Value(*Setting, TEXT("Path="), SettingPath))
 			{
-				SettingPath = SettingPath.TrimQuotes();
+				SettingPath.TrimQuotesInline();
+				SettingPath.ReplaceEscapedCharWithCharInline();
 				if (!SettingPath.IsEmpty())
 				{
 					DataPath = SettingPath;
