Avoid use of the socket subsystem's GetAddressFromString API because it doesn't do name resolving and limits us to IP addresses when connecting to ZenServer.  Name resolving would cause potential stalls in the "IsLocalHost" check, so instead we're limiting ourselves to only considering a set of well-known strings as being localhost: {"localhost, "127.0.0.1", "[::1]"}.

diff --git a/Engine/Source/Developer/Zen/Private/ZenServerInterface.cpp b/Engine/Source/Developer/Zen/Private/ZenServerInterface.cpp
--- a/Engine/Source/Developer/Zen/Private/ZenServerInterface.cpp
+++ b/Engine/Source/Developer/Zen/Private/ZenServerInterface.cpp
@@ -29,7 +29,6 @@
 #include "Serialization/JsonReader.h"
 #include "Serialization/JsonSerializer.h"
 #include "Serialization/JsonWriter.h"
-#include "SocketSubsystem.h"
 #include "String/LexFromString.h"
 #include "ZenVersion.h"
 
@@ -314,30 +313,6 @@ IsLocalHost(const FString& Host)
 		return true;
 	}
 
-	ISocketSubsystem& SocketSubsystem = *ISocketSubsystem::Get();
-
-	const TSharedPtr<FInternetAddr> Addr = SocketSubsystem.GetAddressFromString(Host);
-	if (!Addr)
-	{
-		UE_LOG(LogZenServiceInstance, Warning, TEXT("Failed to get internet address from host '%s'"), *Host);
-		return false;
-	}
-
-	TArray<TSharedPtr<FInternetAddr>> LocalAddresses;
-	if (!SocketSubsystem.GetLocalAdapterAddresses(LocalAddresses))
-	{
-		UE_LOG(LogZenServiceInstance, Warning, TEXT("Failed to find local adapter addresses"));
-		return false;
-	}
-
-	for (const auto& Local : LocalAddresses)
-	{
-		if (*Local == *Addr)
-		{
-			return true;
-		}
-	}
-
 	return false;
 }
 
