Ensure we properly wait for running zen server executable to exit when upgrading

diff --git a/Engine/Source/Developer/Zen/Private/ZenServerInterface.cpp b/Engine/Source/Developer/Zen/Private/ZenServerInterface.cpp
--- a/Engine/Source/Developer/Zen/Private/ZenServerInterface.cpp
+++ b/Engine/Source/Developer/Zen/Private/ZenServerInterface.cpp
@@ -614,7 +614,7 @@ IsLockFileLocked(const TCHAR* FileName, bool bAttemptCleanUp=false)
 #endif
 }
 
-static void
+static bool
 RequestZenShutdownOnPort(uint16 Port)
 {
 #if PLATFORM_WINDOWS
@@ -622,8 +622,10 @@ RequestZenShutdownOnPort(uint16 Port)
 	if (Handle != INVALID_HANDLE_VALUE)
 	{
 		ON_SCOPE_EXIT{ CloseHandle(Handle); };
-		SetEvent(Handle);
+		BOOL OK = SetEvent(Handle);
+		return OK;
 	}
+	return false;
 #elif PLATFORM_UNIX || PLATFORM_MAC
 	TAnsiStringBuilder<64> EventPath;
 	EventPath << "/tmp/Zen_" << Port << "_Shutdown";
@@ -631,19 +633,21 @@ RequestZenShutdownOnPort(uint16 Port)
 	key_t IpcKey = ftok(EventPath.ToString(), 1);
 	if (IpcKey < 0)
 	{
-		return;
+		return false;
 	}
 
 	int Semaphore = semget(IpcKey, 1, 0600);
 	if (Semaphore < 0)
 	{
-		return;
+		return false;
 	}
 
 	semctl(Semaphore, 0, SETVAL, 0);
 	semctl(Semaphore, 0, IPC_RMID);
+	return true;
 #else
 	static_assert(false, "Missing implementation for Zen named shutdown events");
+	return false;
 #endif
 }
 
@@ -850,6 +854,8 @@ StartLocalService(const FZenLocalServiceRunContext& Context, const TCHAR* Transi
 		Parms.Appendf(TEXT(" %s"), TransientArgs);
 	}
 
+	UE_LOG(LogZenServiceInstance, Display, TEXT("Launching executable '%s', working dir '%s', data dir '%s', args '%s'"), *Context.GetExecutable(), *Context.GetWorkingDirectory(), *Context.GetDataPath(), *Parms);
+
 	FProcHandle Proc;
 #if PLATFORM_WINDOWS
 	FString PlatformExecutable = Context.GetExecutable();
@@ -937,7 +943,11 @@ StopLocalService(const TCHAR* DataPath, double MaximumWaitDurationSeconds)
 			return false;
 		}
 
-		RequestZenShutdownOnPort(CurrentPort);
+		if (!RequestZenShutdownOnPort(CurrentPort))
+		{
+			return false;
+		}
+		UE_LOG(LogZenServiceInstance, Display, TEXT("Waiting for running instance to shut down"));
 		return WaitForZenShutdown(*FPaths::Combine(DataPath, TEXT(".lock")), MaximumWaitDurationSeconds);
 	}
 	return true;
@@ -1211,6 +1221,7 @@ FZenServiceInstance::ConditionalUpdateLocalInstall()
 	FString InTreeVersionCache, InstallVersionCache;
 	if (IsInstallVersionOutOfDate(InTreeUtilityPath, InstallUtilityPath, InTreeServicePath, InstallServicePath, InTreeVersionCache, InstallVersionCache))
 	{
+		UE_LOG(LogZenServiceInstance, Display, TEXT("Installing service from '%s' to '%s'"), *InTreeVersionCache, *InstallVersionCache);
 		if (IsProcessActive(*InstallServicePath))
 		{
 			FString DataPath = Settings.SettingsVariant.Get<FServiceAutoLaunchSettings>().DataPath;
@@ -1273,45 +1284,53 @@ FZenServiceInstance::AutoLaunch(const FServiceAutoLaunchSettings& InSettings, FS
 	if (IsLockFileLocked(*LockFilePath, true))
 	{
 		// If an instance is running with this data path, check if we can use it and what port it is on
+		bool bIsReady = false;
 		uint16 CurrentPort = 0;
 		FCbObject LockObject;
 		if (ReadCbLockFile(LockFilePath, LockObject))
 		{
-			bool bIsReady = LockObject["ready"].AsBool();
+			bIsReady = LockObject["ready"].AsBool();
 			if (bIsReady)
 			{
 				CurrentPort = LockObject["port"].AsUInt16();
 			}
 		}
 
-		bool bCurrentInstanceUsable = false;
-
-		FZenLocalServiceRunContext DesiredRunContext;
-		DesiredRunContext.Executable = ExecutablePath;
-		DesiredRunContext.CommandlineArguments = DetermineCmdLineWithoutTransientComponents(InSettings, CurrentPort);
-		DesiredRunContext.WorkingDirectory = WorkingDirectory;
-		DesiredRunContext.DataPath = InSettings.DataPath;
-		DesiredRunContext.bShowConsole = InSettings.bShowConsole;
-
-		FZenLocalServiceRunContext CurrentRunContext;
-		if (CurrentRunContext.ReadFromJsonFile(*ExecutionContextFilePath) && (DesiredRunContext == CurrentRunContext))
+		if (bIsReady)
 		{
-			DesiredPort = CurrentPort;
-			bReUsingExistingInstance = true;
-		}
-		else
-		{
-			RequestZenShutdownOnPort(CurrentPort);
-			WaitForZenShutdown(*LockFilePath, 5.0);
+			FZenLocalServiceRunContext DesiredRunContext;
+			DesiredRunContext.Executable = ExecutablePath;
+			DesiredRunContext.CommandlineArguments = DetermineCmdLineWithoutTransientComponents(InSettings, CurrentPort);
+			DesiredRunContext.WorkingDirectory = WorkingDirectory;
+			DesiredRunContext.DataPath = InSettings.DataPath;
+			DesiredRunContext.bShowConsole = InSettings.bShowConsole;
+
+			FZenLocalServiceRunContext CurrentRunContext;
+			if (CurrentRunContext.ReadFromJsonFile(*ExecutionContextFilePath) && (DesiredRunContext == CurrentRunContext))
+			{
+				DesiredPort = CurrentPort;
+				bReUsingExistingInstance = true;
+			}
+			else
+			{
+				if (RequestZenShutdownOnPort(CurrentPort))
+				{
+					UE_LOG(LogZenServiceInstance, Display, TEXT("Waiting for running instance on port %d to shut down"), CurrentPort);
+					WaitForZenShutdown(*LockFilePath, 5.0);
+				}
+			}
 		}
 	}
 
 	if (!bReUsingExistingInstance)
 	{
-		RequestZenShutdownOnPort(DesiredPort);
+		if (RequestZenShutdownOnPort(DesiredPort))
+		{
+			UE_LOG(LogZenServiceInstance, Display, TEXT("Waiting for running instance on port %d to shut down"), DesiredPort);
+			WaitForZenShutdown(*LockFilePath, 5.0);
+		}
 	}
 
-
 	bool bProcessIsLive = IsLockFileLocked(*LockFilePath);
 
 	// When limiting process lifetime, always re-launch to add sponsor process IDs.
@@ -1349,7 +1368,6 @@ FZenServiceInstance::AutoLaunch(const FServiceAutoLaunchSettings& InSettings, FS
 
 	if (bProcessIsLive)
 	{
-
 		FScopedSlowTask WaitForZenReadySlowTask(0, NSLOCTEXT("Zen", "Zen_WaitingForReady", "Waiting for ZenServer to be ready"));
 		uint64 ZenWaitStartTime = FPlatformTime::Cycles64();
 		enum class EWaitDurationPhase
