Fix race condition when shutting down zenserver via zen autolaunch.

diff --git a/Engine/Source/Developer/Zen/Private/ZenServerInterface.cpp b/Engine/Source/Developer/Zen/Private/ZenServerInterface.cpp
--- a/Engine/Source/Developer/Zen/Private/ZenServerInterface.cpp
+++ b/Engine/Source/Developer/Zen/Private/ZenServerInterface.cpp
@@ -614,12 +614,45 @@ IsLockFileLocked(const TCHAR* FileName, bool bAttemptCleanUp=false)
 #endif
 }
 
+static bool
+IsZenProcessUsingPort(uint16 Port)
+{
+#if PLATFORM_WINDOWS
+	HANDLE Handle = OpenEventW(READ_CONTROL, false, *WriteToWideString<64>(WIDETEXT("Zen_"), Port, WIDETEXT("_Shutdown")));
+	if (Handle != NULL)
+	{
+		ON_SCOPE_EXIT{ CloseHandle(Handle); };
+		return true;
+	}
+	return false;
+#elif PLATFORM_UNIX || PLATFORM_MAC
+	TAnsiStringBuilder<64> EventPath;
+	EventPath << "/tmp/Zen_" << Port << "_Shutdown";
+
+	key_t IpcKey = ftok(EventPath.ToString(), 1);
+	if (IpcKey < 0)
+	{
+		return false;
+	}
+
+	int Semaphore = semget(IpcKey, 1, 0400);
+	if (Semaphore < 0)
+	{
+		return false;
+	}
+	return true;
+#else
+	static_assert(false, "Missing implementation for Zen named shutdown events");
+	return false;
+#endif
+}
+
 static bool
 RequestZenShutdownOnPort(uint16 Port)
 {
 #if PLATFORM_WINDOWS
 	HANDLE Handle = OpenEventW(EVENT_MODIFY_STATE, false, *WriteToWideString<64>(WIDETEXT("Zen_"), Port, WIDETEXT("_Shutdown")));
-	if (Handle != INVALID_HANDLE_VALUE)
+	if (Handle != NULL)
 	{
 		ON_SCOPE_EXIT{ CloseHandle(Handle); };
 		BOOL OK = SetEvent(Handle);
@@ -652,10 +685,10 @@ RequestZenShutdownOnPort(uint16 Port)
 }
 
 static bool
-WaitForZenShutdown(const TCHAR* LockFilePath, double MaximumWaitDurationSeconds)
+WaitForZenShutdown(const TCHAR* LockFilePath, uint16 CurrentPort, double MaximumWaitDurationSeconds)
 {
 	uint64 ZenShutdownWaitStartTime = FPlatformTime::Cycles64();
-	while (IsLockFileLocked(LockFilePath))
+	while (IsLockFileLocked(LockFilePath) || IsZenProcessUsingPort(CurrentPort))
 	{
 		double ZenShutdownWaitDuration = FPlatformTime::ToSeconds64(FPlatformTime::Cycles64() - ZenShutdownWaitStartTime);
 		if (ZenShutdownWaitDuration < MaximumWaitDurationSeconds)
@@ -948,7 +981,7 @@ StopLocalService(const TCHAR* DataPath, double MaximumWaitDurationSeconds)
 			return false;
 		}
 		UE_LOG(LogZenServiceInstance, Display, TEXT("Waiting for running instance to shut down"));
-		return WaitForZenShutdown(*FPaths::Combine(DataPath, TEXT(".lock")), MaximumWaitDurationSeconds);
+		return WaitForZenShutdown(*FPaths::Combine(DataPath, TEXT(".lock")), CurrentPort, MaximumWaitDurationSeconds);
 	}
 	return true;
 }
@@ -1316,7 +1349,7 @@ FZenServiceInstance::AutoLaunch(const FServiceAutoLaunchSettings& InSettings, FS
 				if (RequestZenShutdownOnPort(CurrentPort))
 				{
 					UE_LOG(LogZenServiceInstance, Display, TEXT("Waiting for running instance on port %d to shut down"), CurrentPort);
-					WaitForZenShutdown(*LockFilePath, 5.0);
+					WaitForZenShutdown(*LockFilePath, CurrentPort, 5.0);
 				}
 			}
 		}
@@ -1327,11 +1360,11 @@ FZenServiceInstance::AutoLaunch(const FServiceAutoLaunchSettings& InSettings, FS
 		if (RequestZenShutdownOnPort(DesiredPort))
 		{
 			UE_LOG(LogZenServiceInstance, Display, TEXT("Waiting for running instance on port %d to shut down"), DesiredPort);
-			WaitForZenShutdown(*LockFilePath, 5.0);
+			WaitForZenShutdown(*LockFilePath, DesiredPort, 5.0);
 		}
 	}
 
-	bool bProcessIsLive = IsLockFileLocked(*LockFilePath);
+	bool bProcessIsLive = IsLockFileLocked(*LockFilePath) || IsZenProcessUsingPort(DesiredPort);
 
 	// When limiting process lifetime, always re-launch to add sponsor process IDs.
 	// When not limiting process lifetime, only launch if the process is not already live.
