Improved logic for installing new version of zenserver.
- Use shared ZenServerState to get zen instance state to better identify processes and port usage
- Check is the process is running using Pid and process handle instead of checking for locked files
- Make sure we update cached zen.version for InTree version even if there is no current installation
- Add `-ForceZenInstall` option to force an install of InTree version even if installed version is the same or newer (enables downgrading)
- More logging added

diff --git a/Engine/Source/Developer/Zen/Private/ZenServerInterface.cpp b/Engine/Source/Developer/Zen/Private/ZenServerInterface.cpp
--- a/Engine/Source/Developer/Zen/Private/ZenServerInterface.cpp
+++ b/Engine/Source/Developer/Zen/Private/ZenServerInterface.cpp
@@ -42,6 +42,7 @@
 
 #if PLATFORM_UNIX || PLATFORM_MAC
 #	include <sys/file.h>
+#	include <sys/mman.h>
 #	include <sys/sem.h>
 #endif
 
@@ -52,6 +53,152 @@ namespace UE::Zen
 
 DEFINE_LOG_CATEGORY_STATIC(LogZenServiceInstance, Log, All);
 
+class ZenServerState
+{
+public:
+	ZenServerState();
+	~ZenServerState();
+
+	struct ZenServerEntry
+	{
+		// This matches the structure found in the Zen server
+		// https://github.com/EpicGames/zen/blob/main/zenutil/include/zenutil/zenserverprocess.h#L91
+		//
+		std::atomic<uint32> Pid;
+		std::atomic<uint16> DesiredListenPort;
+		std::atomic<uint16> Flags;
+		uint8				  SessionId[12];
+		std::atomic<uint32> SponsorPids[8];
+		std::atomic<uint16> EffectiveListenPort;
+		uint8				  Padding[10];
+
+		enum class FlagsEnum : uint16
+		{
+			kShutdownPlease = 1 << 0,
+			kIsReady = 1 << 1,
+		};
+	};
+	static_assert(sizeof(ZenServerEntry) == 64);
+
+	const ZenServerEntry* LookupByPort(int DesiredListenPort);
+	const ZenServerEntry* LookupByPid(uint32 Pid);
+
+private:
+	void* m_hMapFile = nullptr;
+	const ZenServerEntry* m_Data = nullptr;
+	int				m_MaxEntryCount = 65536 / sizeof(ZenServerEntry);
+};
+
+ZenServerState::ZenServerState()
+	: m_hMapFile(nullptr)
+	, m_Data(nullptr)
+{
+	size_t MapSize = m_MaxEntryCount * sizeof(ZenServerEntry);
+
+#if PLATFORM_WINDOWS
+	HANDLE hMap = OpenFileMapping(FILE_MAP_READ, 0, L"Global\\ZenMap");
+	if (hMap == NULL)
+	{
+		hMap = OpenFileMapping(FILE_MAP_READ, 0, L"Local\\ZenMap");
+	}
+
+	if (hMap == NULL)
+	{
+		return;
+	}
+
+	const void* pBuf = MapViewOfFile(hMap,		   // handle to map object
+		FILE_MAP_READ,  // read permission
+		0,			   // offset high
+		0,			   // offset low
+		MapSize);
+
+	if (pBuf == NULL)
+	{
+		CloseHandle(hMap);
+		return;
+	}
+#elif PLATFORM_UNIX || PLATFORM_MAC
+	int Fd = shm_open("/UnrealEngineZen", O_RDONLY | O_CLOEXEC, 0666);
+	if (Fd < 0)
+	{
+		return;
+	}
+	void* hMap = (void*)intptr_t(Fd);
+
+	const void* pBuf = mmap(nullptr, MapSize, PROT_READ, MAP_PRIVATE, Fd, 0);
+	if (pBuf == MAP_FAILED)
+	{
+		close(Fd);
+		return;
+	}
+#endif
+
+	m_hMapFile = hMap;
+	m_Data = reinterpret_cast<const ZenServerEntry*>(pBuf);
+}
+
+ZenServerState::~ZenServerState()
+{
+#if PLATFORM_WINDOWS
+	if (m_Data)
+	{
+		UnmapViewOfFile(m_Data);
+	}
+
+	if (m_hMapFile)
+	{
+		CloseHandle(m_hMapFile);
+	}
+#elif PLATFORM_UNIX || PLATFORM_MAC
+	if (m_Data != nullptr)
+	{
+		munmap((void*)m_Data, m_MaxEntryCount * sizeof(ZenServerEntry));
+	}
+
+	int Fd = int(intptr_t(m_hMapFile));
+	close(Fd);
+#endif
+	m_hMapFile = nullptr;
+	m_Data = nullptr;
+}
+
+const ZenServerState::ZenServerEntry* ZenServerState::LookupByPort(int DesiredListenPort)
+{
+	if (m_Data == nullptr)
+	{
+		return nullptr;
+	}
+
+	for (int i = 0; i < m_MaxEntryCount; ++i)
+	{
+		if (m_Data[i].DesiredListenPort == DesiredListenPort)
+		{
+			return &m_Data[i];
+		}
+	}
+
+	return nullptr;
+}
+
+const ZenServerState::ZenServerEntry* ZenServerState::LookupByPid(uint32 Pid)
+{
+	if (m_Data == nullptr)
+	{
+		return nullptr;
+	}
+
+	for (int i = 0; i < m_MaxEntryCount; ++i)
+	{
+		if (m_Data[i].Pid == Pid)
+		{
+			return &m_Data[i];
+		}
+	}
+
+	return nullptr;
+}
+
 static FZenVersion
 GetZenVersion(const FString& UtilityPath, const FString& ServicePath, const FString& VersionCachePath)
 {
@@ -72,10 +219,12 @@ GetZenVersion(const FString& UtilityPath, const FString& ServicePath, const FStr
 	};
 
 	FDateTime VersionCacheModificationTime = FileManager.GetTimeStamp(*VersionCachePath);
+	bool VersionCacheIsOlderThanUtilityExecutable = VersionCacheModificationTime < UtilityExecutableModificationTime;
+	bool VersionCacheIsOlderThanServerExecutable = VersionCacheModificationTime < ServiceExecutableModificationTime;
+	bool VersionCacheIsOutOfDate = VersionCacheIsOlderThanUtilityExecutable || VersionCacheIsOlderThanServerExecutable;
 	FString VersionFileContents;
 	FZenVersion ComparableVersion;
-	if ((VersionCacheModificationTime < UtilityExecutableModificationTime) ||
-		(VersionCacheModificationTime < ServiceExecutableModificationTime) ||
+	if (VersionCacheIsOutOfDate ||
 		!FFileHelper::LoadFileToString(VersionFileContents, *VersionCachePath) ||
 		!ComparableVersion.TryParse(*VersionFileContents))
 	{
@@ -123,19 +272,33 @@ GetZenVersion(const FString& UtilityPath, const FString& ServicePath, const FStr
 static bool
 IsInstallVersionOutOfDate(const FString& InTreeUtilityPath, const FString& InstallUtilityPath, const FString& InTreeServicePath, const FString& InstallServicePath, FString& OutInTreeVersionCache, FString& OutInstallVersionCache)
 {
-	OutInTreeVersionCache = FPaths::Combine(FPaths::EngineSavedDir(), TEXT("Zen") TEXT("zen.version"));
+	OutInTreeVersionCache = FPaths::Combine(FPaths::EngineSavedDir(), TEXT("Zen"), TEXT("zen.version"));
 	OutInstallVersionCache = FPaths::SetExtension(InstallUtilityPath, TEXT("version"));
 
+	// Always get the InTree utility path so cached version information is up to date
+	FZenVersion InTreeVersion = GetZenVersion(InTreeUtilityPath, InTreeServicePath, OutInTreeVersionCache);
+	UE_LOG(LogZenServiceInstance, Log, TEXT("InTree version at '%s' is '%s'"), *InTreeServicePath, *InTreeVersion.ToString());
+
 	IFileManager& FileManager = IFileManager::Get();
 	if (!FileManager.FileExists(*InstallUtilityPath) || !FileManager.FileExists(*InstallServicePath))
 	{
+		UE_LOG(LogZenServiceInstance, Log, TEXT("No installation found at '%s'"), *InstallServicePath);
 		return true;
 	}
-
-	FZenVersion InTreeVersion = GetZenVersion(InTreeUtilityPath, InTreeServicePath, OutInTreeVersionCache);
 	FZenVersion InstallVersion = GetZenVersion(InstallUtilityPath, InstallServicePath, OutInstallVersionCache);
+	UE_LOG(LogZenServiceInstance, Log, TEXT("Installed version at '%s' is '%s'"), *InstallServicePath, *InstallVersion.ToString());
 
-	return InstallVersion < InTreeVersion;
+	if (InstallVersion < InTreeVersion)
+	{
+		UE_LOG(LogZenServiceInstance, Log, TEXT("Installed version at '%s' (%s) is older than '%s' (%s)"), *InstallServicePath, *InstallVersion.ToString(), *InTreeServicePath, *InTreeVersion.ToString());
+		return true;
+	}
+	if (FParse::Param(FCommandLine::Get(), TEXT("ForceZenInstall")))
+	{
+		UE_LOG(LogZenServiceInstance, Display, TEXT("Forcing install from '%s' (%s) over '%s' (%s)"), *InTreeServicePath, *InTreeVersion.ToString(), *InstallServicePath, *InstallVersion.ToString());
+		return true;
+	}
+	return false;
 }
 
 static bool
@@ -156,8 +319,12 @@ AttemptFileCopyWithRetries(const TCHAR* Dst, const TCHAR* Src, double RetryDurat
 		}
 		CopyResult = IFileManager::Get().Copy(Dst, Src, true, true, false);
 	}
-
-	return CopyResult == COPY_OK;
+	if (CopyResult == COPY_OK)
+	{
+		return true;
+	}
+	UE_LOG(LogZenServiceInstance, Warning, TEXT("%s from '%s' to '%s', '%s'"), Src, Dst, CopyResult == COPY_Fail ? "Failed to copy file" : "Cancelled file copy");
+	return false;
 }
 
 static void
@@ -647,6 +814,55 @@ IsZenProcessUsingPort(uint16 Port)
 #endif
 }
 
+static bool
+IsZenProcessActive(const TCHAR* ExecutablePath, uint32* OutPid)
+{
+	FString NormalizedExecutablePath(ExecutablePath);
+	FPaths::NormalizeFilename(NormalizedExecutablePath);
+	FPlatformProcess::FProcEnumerator ProcIter;
+	while (ProcIter.MoveNext())
+	{
+		FPlatformProcess::FProcEnumInfo ProcInfo = ProcIter.GetCurrent();
+		FString Candidate = ProcInfo.GetFullPath();
+		FPaths::NormalizeFilename(Candidate);
+		if (Candidate == NormalizedExecutablePath)
+		{
+			if (OutPid)
+			{
+				*OutPid = ProcInfo.GetPID();
+			}
+			return true;
+		}
+	}
+	return false;
+}
+
+static bool
+IsZenProcessUsingDataDir(const TCHAR* DataPath, uint16* OutPort)
+{
+	const FString LockFilePath = FPaths::Combine(DataPath, TEXT(".lock"));
+	if (IsLockFileLocked(*LockFilePath, true))
+	{
+		if (OutPort)
+		{
+			// If an instance is running with this data path, check if we can use it and what port it is on
+			*OutPort = 0;
+			FCbObject LockObject;
+			if (ReadCbLockFile(LockFilePath, LockObject))
+			{
+				bool bIsReady = LockObject["ready"].AsBool();
+				if (bIsReady)
+				{
+					*OutPort = LockObject["port"].AsUInt16();
+				}
+			}
+		}
+
+		return true;
+	}
+	return false;
+}
+
 static bool
 RequestZenShutdownOnPort(uint16 Port)
 {
@@ -685,10 +901,38 @@ RequestZenShutdownOnPort(uint16 Port)
 }
 
 static bool
-WaitForZenShutdown(const TCHAR* LockFilePath, uint16 CurrentPort, double MaximumWaitDurationSeconds)
+ShutdownRunningService(const TCHAR* ExecutablePath, double MaximumWaitDurationSeconds = 5.0)
 {
+	uint32 ServicePid;
+	if (!IsZenProcessActive(ExecutablePath, &ServicePid))
+	{
+		UE_LOG(LogZenServiceInstance, Log, TEXT("No running service found for executable '%s'"), ExecutablePath);
+		return true;
+	}
+	UE_LOG(LogZenServiceInstance, Display, TEXT("Waiting for running instance using executable '%s' to shut down"), ExecutablePath);
+	FProcHandle ProcessHandle = FPlatformProcess::OpenProcess(ServicePid);
+	if (!ProcessHandle.IsValid() && IsZenProcessActive(ExecutablePath, nullptr))
+	{
+		UE_LOG(LogZenServiceInstance, Warning, TEXT("Failed to open handle for running service for executable '%s' (Pid: %u)"), ExecutablePath, ServicePid);
+		return false;
+	}
+	ON_SCOPE_EXIT{ FPlatformProcess::CloseProc(ProcessHandle); };
+
+	ZenServerState ServerState;
+	const ZenServerState::ZenServerEntry* Entry = ServerState.LookupByPid(ServicePid);
+	if (!Entry && FPlatformProcess::IsProcRunning(ProcessHandle))
+	{
+		UE_LOG(LogZenServiceInstance, Warning, TEXT("Can't find server state for running service for executable '%s' (Pid: %u)"), ExecutablePath, ServicePid);
+		return false;
+	}
+	uint16 ProcessPort = Entry->DesiredListenPort;
+	if (!RequestZenShutdownOnPort(ProcessPort) && FPlatformProcess::IsProcRunning(ProcessHandle))
+	{
+		UE_LOG(LogZenServiceInstance, Warning, TEXT("Failed to request shutdown for running service for executable '%s' (Pid: %u, Port: %u)"), ExecutablePath, ServicePid, ProcessPort);
+		return false;
+	}
 	uint64 ZenShutdownWaitStartTime = FPlatformTime::Cycles64();
-	while (IsLockFileLocked(LockFilePath) || IsZenProcessUsingPort(CurrentPort))
+	while (FPlatformProcess::IsProcRunning(ProcessHandle))
 	{
 		double ZenShutdownWaitDuration = FPlatformTime::ToSeconds64(FPlatformTime::Cycles64() - ZenShutdownWaitStartTime);
 		if (ZenShutdownWaitDuration < MaximumWaitDurationSeconds)
@@ -697,28 +941,57 @@ WaitForZenShutdown(const TCHAR* LockFilePath, uint16 CurrentPort, double Maximum
 		}
 		else
 		{
+			UE_LOG(LogZenServiceInstance, Warning, TEXT("Timed out waiting for shutdown of running service for executable '%s' (Pid: %u, Port: %u)"), ExecutablePath, ServicePid, ProcessPort);
 			return false;
 		}
 	}
 	return true;
 }
 
-static bool IsProcessActive(const TCHAR* ExecutablePath)
+static bool
+ShutdownRunningService(uint16 ProcessPort, double MaximumWaitDurationSeconds = 5.0)
 {
-	FString NormalizedExecutablePath(ExecutablePath);
-	FPaths::NormalizeFilename(NormalizedExecutablePath);
-	FPlatformProcess::FProcEnumerator ProcIter;
-	while (ProcIter.MoveNext())
+	if (!IsZenProcessUsingPort(ProcessPort))
 	{
-		FPlatformProcess::FProcEnumInfo ProcInfo = ProcIter.GetCurrent();
-		FString Candidate = ProcInfo.GetFullPath();
-		FPaths::NormalizeFilename(Candidate);
-		if (Candidate == NormalizedExecutablePath)
+		return true;
+	}
+	UE_LOG(LogZenServiceInstance, Display, TEXT("Waiting for running instance using port %u to shut down"), ProcessPort);
+	ZenServerState ServerState;
+	const ZenServerState::ZenServerEntry* Entry = ServerState.LookupByPort(ProcessPort);
+	if (!Entry && IsZenProcessUsingPort(ProcessPort))
+	{
+		UE_LOG(LogZenServiceInstance, Warning, TEXT("Can't find server state for running service using port %u"), ProcessPort);
+		return false;
+	}
+	uint32 ServicePid = Entry->Pid;
+	FProcHandle ProcessHandle = FPlatformProcess::OpenProcess(ServicePid);
+	if (!ProcessHandle.IsValid() && IsZenProcessUsingPort(ProcessPort))
+	{
+		UE_LOG(LogZenServiceInstance, Warning, TEXT("Failed to open handle for running service using port %u (Pid: %u)"), ProcessPort, ServicePid);
+		return false;
+	}
+	ON_SCOPE_EXIT{ FPlatformProcess::CloseProc(ProcessHandle); };
+
+	if (!RequestZenShutdownOnPort(ProcessPort) && FPlatformProcess::IsProcRunning(ProcessHandle))
+	{
+		UE_LOG(LogZenServiceInstance, Warning, TEXT("Failed to request shutdown for running service using port %u (Pid: %u)"), ProcessPort, ServicePid);
+		return false;
+	}
+	uint64 ZenShutdownWaitStartTime = FPlatformTime::Cycles64();
+	while (FPlatformProcess::IsProcRunning(ProcessHandle))
+	{
+		double ZenShutdownWaitDuration = FPlatformTime::ToSeconds64(FPlatformTime::Cycles64() - ZenShutdownWaitStartTime);
+		if (ZenShutdownWaitDuration < MaximumWaitDurationSeconds)
 		{
-			return true;
+			FPlatformProcess::Sleep(0.01f);
+		}
+		else
+		{
+			UE_LOG(LogZenServiceInstance, Warning, TEXT("Timed out waiting for shutdown of running service using port %u (Pid: %u)"), ProcessPort, ServicePid);
+			return false;
 		}
 	}
-	return false;
+	return true;
 }
 
 static FString
@@ -854,28 +1127,7 @@ FZenLocalServiceRunContext::WriteToJsonFile(const TCHAR* Filename) const
 bool
 IsLocalServiceRunning(const TCHAR* DataPath, uint16* OutPort)
 {
-	// TODO: Instead of using the lock file, this could use the shared memory system state named "Global\ZenMap" (see zenserverprocess.{h,cpp} in zen codebase)
-	const FString LockFilePath = FPaths::Combine(DataPath, TEXT(".lock"));
-	if (IsLockFileLocked(*LockFilePath, true))
-	{
-		if (OutPort)
-		{
-			// If an instance is running with this data path, check if we can use it and what port it is on
-			*OutPort = 0;
-			FCbObject LockObject;
-			if (ReadCbLockFile(LockFilePath, LockObject))
-			{
-				bool bIsReady = LockObject["ready"].AsBool();
-				if (bIsReady)
-				{
-					*OutPort = LockObject["port"].AsUInt16();
-				}
-			}
-		}
-
-		return true;
-	}
-	return false;
+	return IsZenProcessUsingDataDir(DataPath, OutPort);
 }
 
 FProcHandle
@@ -971,17 +1223,7 @@ StopLocalService(const TCHAR* DataPath, double MaximumWaitDurationSeconds)
 	uint16 CurrentPort = 0;
 	if (IsLocalServiceRunning(DataPath, &CurrentPort))
 	{
-		if (CurrentPort == 0)
-		{
-			return false;
-		}
-
-		if (!RequestZenShutdownOnPort(CurrentPort))
-		{
-			return false;
-		}
-		UE_LOG(LogZenServiceInstance, Display, TEXT("Waiting for running instance to shut down"));
-		return WaitForZenShutdown(*FPaths::Combine(DataPath, TEXT(".lock")), CurrentPort, MaximumWaitDurationSeconds);
+		return ShutdownRunningService(CurrentPort, MaximumWaitDurationSeconds);
 	}
 	return true;
 }
@@ -1158,7 +1400,10 @@ FZenServiceInstance::TryRecovery()
 			FZenLocalServiceRunContext RunContext;
 			if (TryGetLocalServiceRunContext(RunContext))
 			{
-				StopLocalService(*RunContext.GetDataPath());
+				if (!ShutdownRunningService(*RunContext.GetDataPath()))
+				{
+					return false;
+				}
 				StartLocalService(RunContext);
 				UE_LOG(LogZenServiceInstance, Display, TEXT("Local ZenServer recovery finished."));
 			}
@@ -1208,11 +1453,15 @@ FZenServiceInstance::Initialize()
 {
 	if (Settings.IsAutoLaunch())
 	{
-		bHasLaunchedLocal = AutoLaunch(Settings.SettingsVariant.Get<FServiceAutoLaunchSettings>(), ConditionalUpdateLocalInstall(), HostName, Port);
-		if (bHasLaunchedLocal)
+		FString ExecutableInstallPath = ConditionalUpdateLocalInstall();
+		if (!ExecutableInstallPath.IsEmpty())
 		{
-			AutoLaunchedPort = Port;
-			bIsRunningLocally = true;
+			bHasLaunchedLocal = AutoLaunch(Settings.SettingsVariant.Get<FServiceAutoLaunchSettings>(), *ExecutableInstallPath, HostName, Port);
+			if (bHasLaunchedLocal)
+			{
+				AutoLaunchedPort = Port;
+				bIsRunningLocally = true;
+			}
 		}
 	}
 	else
@@ -1230,7 +1479,8 @@ FZenServiceInstance::PromptUserToStopRunningServerInstance(const FString& Server
 {
 	if (FApp::IsUnattended())
 	{
-		// Do not ask if there is no one to show a message
+		// Just log as there is no one to show a message
+		UE_LOG(LogZenServiceInstance, Display, TEXT("ZenServer needs to be updated to a new version. Please shut down any tools that are using the ZenServer at '%s'"), *ServerFilePath);
 		return;
 	}
 
@@ -1254,31 +1504,29 @@ FZenServiceInstance::ConditionalUpdateLocalInstall()
 	FString InTreeVersionCache, InstallVersionCache;
 	if (IsInstallVersionOutOfDate(InTreeUtilityPath, InstallUtilityPath, InTreeServicePath, InstallServicePath, InTreeVersionCache, InstallVersionCache))
 	{
-		UE_LOG(LogZenServiceInstance, Display, TEXT("Installing service from '%s' to '%s'"), *InTreeVersionCache, *InstallVersionCache);
-		if (IsProcessActive(*InstallServicePath))
+		UE_LOG(LogZenServiceInstance, Display, TEXT("Installing service from '%s' to '%s'"), *InTreeServicePath, *InstallServicePath);
+		if (!ShutdownRunningService(*InstallServicePath))
 		{
-			FString DataPath = Settings.SettingsVariant.Get<FServiceAutoLaunchSettings>().DataPath;
-			StopLocalService(*DataPath);
-
-			if (IsLocalServiceRunning(*DataPath))
-			{
-				PromptUserToStopRunningServerInstance(InstallServicePath);
-			}
+			PromptUserToStopRunningServerInstance(InstallServicePath);
+			return FString();
 		}
 
-		// Even after waiting for the lock file to be removed, the executable may have a period where it can't be overwritten as the process shuts down
-		// so any attempt to overwrite it should have some tolerance for retrying.
-		bool bServiceExecutableUpdated = AttemptFileCopyWithRetries(*InstallServicePath, *InTreeServicePath, 5.0);
-		checkf(bServiceExecutableUpdated, TEXT("Failed to copy zenserver to install location '%s'."), *InstallServicePath);
-
-		bool bUtilityExecutableUpdated = AttemptFileCopyWithRetries(*InstallUtilityPath, *InTreeUtilityPath, 5.0);
-		checkf(bUtilityExecutableUpdated, TEXT("Failed to copy zen to install location '%s'."), *InstallUtilityPath);
+		// Even after waiting for the process to shut down we have a tolerance for failure when overwriting the target files
+		if (!AttemptFileCopyWithRetries(*InstallServicePath, *InTreeServicePath, 5.0))
+		{
+			PromptUserToStopRunningServerInstance(InstallServicePath);
+			return FString();
+		}
 
-		bMainExecutablesUpdated = (bServiceExecutableUpdated && bUtilityExecutableUpdated);
-		if (bMainExecutablesUpdated)
+		if (!AttemptFileCopyWithRetries(*InstallUtilityPath, *InTreeUtilityPath, 5.0))
 		{
-			AttemptFileCopyWithRetries(*InstallVersionCache, *InTreeVersionCache, 1.0);
+			PromptUserToStopRunningServerInstance(InstallServicePath);
+			return FString();
 		}
+
+		AttemptFileCopyWithRetries(*InstallVersionCache, *InTreeVersionCache, 1.0);
+
+		bMainExecutablesUpdated = true;
 	}
 
 #if PLATFORM_WINDOWS
@@ -1305,7 +1553,7 @@ FZenServiceInstance::ConditionalUpdateLocalInstall()
 bool
 FZenServiceInstance::AutoLaunch(const FServiceAutoLaunchSettings& InSettings, FString&& ExecutablePath, FString& OutHostName, uint16& OutPort)
 {
-	int16 DesiredPort = InSettings.DesiredPort;
+	uint16 DesiredPort = InSettings.DesiredPort;
 	IFileManager& FileManager = IFileManager::Get();
 	const FString LockFilePath = FPaths::Combine(InSettings.DataPath, TEXT(".lock"));
 	const FString ExecutionContextFilePath = FPaths::SetExtension(ExecutablePath, TEXT(".runcontext"));
@@ -1314,22 +1562,12 @@ FZenServiceInstance::AutoLaunch(const FServiceAutoLaunchSettings& InSettings, FS
 
 	bool bReUsingExistingInstance = false;
 
-	if (IsLockFileLocked(*LockFilePath, true))
+	// When limiting process lifetime, always re-launch to add sponsor process IDs.
+	// When not limiting process lifetime, only launch if the process is not already live.
+	if (!InSettings.bLimitProcessLifetime)
 	{
-		// If an instance is running with this data path, check if we can use it and what port it is on
-		bool bIsReady = false;
-		uint16 CurrentPort = 0;
-		FCbObject LockObject;
-		if (ReadCbLockFile(LockFilePath, LockObject))
-		{
-			bIsReady = LockObject["ready"].AsBool();
-			if (bIsReady)
-			{
-				CurrentPort = LockObject["port"].AsUInt16();
-			}
-		}
-
-		if (bIsReady)
+		uint16 CurrentPort;
+		if (IsZenProcessUsingDataDir(*InSettings.DataPath, &CurrentPort) && CurrentPort != 0)
 		{
 			FZenLocalServiceRunContext DesiredRunContext;
 			DesiredRunContext.Executable = ExecutablePath;
@@ -1341,35 +1579,23 @@ FZenServiceInstance::AutoLaunch(const FServiceAutoLaunchSettings& InSettings, FS
 			FZenLocalServiceRunContext CurrentRunContext;
 			if (CurrentRunContext.ReadFromJsonFile(*ExecutionContextFilePath) && (DesiredRunContext == CurrentRunContext))
 			{
+				UE_LOG(LogZenServiceInstance, Log, TEXT("Reusing existing instance running on port %u"), CurrentPort);
 				DesiredPort = CurrentPort;
 				bReUsingExistingInstance = true;
 			}
-			else
-			{
-				if (RequestZenShutdownOnPort(CurrentPort))
-				{
-					UE_LOG(LogZenServiceInstance, Display, TEXT("Waiting for running instance on port %d to shut down"), CurrentPort);
-					WaitForZenShutdown(*LockFilePath, CurrentPort, 5.0);
-				}
-			}
 		}
 	}
 
 	if (!bReUsingExistingInstance)
 	{
-		if (RequestZenShutdownOnPort(DesiredPort))
+		if (!ShutdownRunningService(DesiredPort))
 		{
-			UE_LOG(LogZenServiceInstance, Display, TEXT("Waiting for running instance on port %d to shut down"), DesiredPort);
-			WaitForZenShutdown(*LockFilePath, DesiredPort, 5.0);
+			UE_LOG(LogZenServiceInstance, Warning, TEXT("Failed to shut down running service using port %u"), DesiredPort);
+			return false;
 		}
-	}
+		checkf(!IsLockFileLocked(*LockFilePath), TEXT("Data folder is still locked by other process '%s'"), *InSettings.DataPath);
+		checkf(!IsZenProcessUsingPort(DesiredPort), TEXT("Service port is still in use %u"), DesiredPort);
 
-	bool bProcessIsLive = IsLockFileLocked(*LockFilePath) || IsZenProcessUsingPort(DesiredPort);
-
-	// When limiting process lifetime, always re-launch to add sponsor process IDs.
-	// When not limiting process lifetime, only launch if the process is not already live.
-	if (InSettings.bLimitProcessLifetime || !bProcessIsLive)
-	{
 		FString ParmsWithoutTransients = DetermineCmdLineWithoutTransientComponents(InSettings, DesiredPort);
 		FString TransientParms;
 
@@ -1386,107 +1612,98 @@ FZenServiceInstance::AutoLaunch(const FServiceAutoLaunchSettings& InSettings, FS
 		EffectiveRunContext.bShowConsole = InSettings.bShowConsole;
 
 		FProcHandle Proc = StartLocalService(EffectiveRunContext, TransientParms.IsEmpty() ? nullptr : *TransientParms);
-		if (!bProcessIsLive)
+		EffectiveRunContext.WriteToJsonFile(*ExecutionContextFilePath);
+
+		if (!Proc.IsValid())
 		{
-			EffectiveRunContext.WriteToJsonFile(*ExecutionContextFilePath);
+			UE_LOG(LogZenServiceInstance, Warning, TEXT("Failed launch service using executable '%s' on port %u"), *ExecutablePath, DesiredPort);
+			return false;
 		}
-
-		bProcessIsLive = Proc.IsValid();
 	}
 
-
 	OutHostName = TEXT("[::1]");
 	// Default to assuming that we get to run on the port we want
 	OutPort = DesiredPort;
 
-	if (bProcessIsLive)
+	FScopedSlowTask WaitForZenReadySlowTask(0, NSLOCTEXT("Zen", "Zen_WaitingForReady", "Waiting for ZenServer to be ready"));
+	uint64 ZenWaitStartTime = FPlatformTime::Cycles64();
+	enum class EWaitDurationPhase
 	{
-		FScopedSlowTask WaitForZenReadySlowTask(0, NSLOCTEXT("Zen", "Zen_WaitingForReady", "Waiting for ZenServer to be ready"));
-		uint64 ZenWaitStartTime = FPlatformTime::Cycles64();
-		enum class EWaitDurationPhase
-		{
-			Short,
-			Medium,
-			Long
-		} DurationPhase = EWaitDurationPhase::Short;
-		bool bIsReady = false;
-		while (!bIsReady)
+		Short,
+		Medium,
+		Long
+	} DurationPhase = EWaitDurationPhase::Short;
+	bool bIsReady = false;
+	while (!bIsReady)
+	{
+		FCbObject LockObject;
+		if (ReadCbLockFile(LockFilePath, LockObject))
 		{
-			FCbObject LockObject;
-			if (ReadCbLockFile(LockFilePath, LockObject))
+			bIsReady = LockObject["ready"].AsBool();
+			if (bIsReady)
 			{
-				bIsReady = LockObject["ready"].AsBool();
-				if (bIsReady)
-				{
-					OutPort = LockObject["port"].AsUInt16(DesiredPort);
-					break;
-				}
+				OutPort = LockObject["port"].AsUInt16(DesiredPort);
+				break;
 			}
+		}
 
-			double ZenWaitDuration = FPlatformTime::ToSeconds64(FPlatformTime::Cycles64() - ZenWaitStartTime);
-			if (ZenWaitDuration < 10.0)
-			{
-				// Initial 10 second window of higher frequency checks
-				FPlatformProcess::Sleep(0.01f);
-			}
-			else
+		double ZenWaitDuration = FPlatformTime::ToSeconds64(FPlatformTime::Cycles64() - ZenWaitStartTime);
+		if (ZenWaitDuration < 10.0)
+		{
+			// Initial 10 second window of higher frequency checks
+			FPlatformProcess::Sleep(0.01f);
+		}
+		else
+		{
+			if (DurationPhase == EWaitDurationPhase::Short)
 			{
-				if (DurationPhase == EWaitDurationPhase::Short)
+				if (!IsLockFileLocked(*LockFilePath))
 				{
-					if (!IsLockFileLocked(*LockFilePath))
+					if (FApp::IsUnattended())
 					{
-						if (FApp::IsUnattended())
-						{
-							checkf(false, TEXT("ZenServer did not launch in the expected duration."));
-						}
-						else
-						{
-							FText ZenLaunchFailurePromptTitle = NSLOCTEXT("Zen", "Zen_LaunchFailurePromptTitle", "Failed to launch");
-
-							FFormatNamedArguments FormatArguments;
-							FString LogFilePath = FPaths::Combine(InSettings.DataPath, TEXT("logs"), TEXT("zenserver.log"));
-							FPaths::MakePlatformFilename(LogFilePath);
-							FormatArguments.Add(TEXT("LogFilePath"), FText::FromString(LogFilePath));
-							FText ZenLaunchFailurePromptText = FText::Format(NSLOCTEXT("Zen", "Zen_LaunchFailurePromptText", "ZenServer failed to launch. This process will now exit. Please check the ZenServer log file for details:\n{LogFilePath}"), FormatArguments);
-							FPlatformMisc::MessageBoxExt(EAppMsgType::Ok, *ZenLaunchFailurePromptText.ToString(), *ZenLaunchFailurePromptTitle.ToString());
-							FPlatformMisc::RequestExit(true);
-							return false;
-						}
+						checkf(false, TEXT("ZenServer did not launch in the expected duration."));
 					}
-					// Note that the dialog may not show up when zenserver is needed early in the launch cycle, but this will at least ensure
-					// the splash screen is refreshed with the appropriate text status message.
-					WaitForZenReadySlowTask.MakeDialog(true, false);
-					UE_LOG(LogZenServiceInstance, Display, TEXT("Waiting for ZenServer to be ready..."));
-					DurationPhase = EWaitDurationPhase::Medium;
-				}
-				else if (!FApp::IsUnattended() && ZenWaitDuration > 20.0 && (DurationPhase == EWaitDurationPhase::Medium))
-				{
-					FText ZenLongWaitPromptTitle = NSLOCTEXT("Zen", "Zen_LongWaitPromptTitle", "Wait for ZenServer?");
-					FText ZenLongWaitPromptText = NSLOCTEXT("Zen", "Zen_LongWaitPromptText", "ZenServer is taking a long time to launch. It may be performing maintenance. Keep waiting?");
-					if (FPlatformMisc::MessageBoxExt(EAppMsgType::YesNo, *ZenLongWaitPromptText.ToString(), *ZenLongWaitPromptTitle.ToString()) == EAppReturnType::No)
+					else
 					{
+						FText ZenLaunchFailurePromptTitle = NSLOCTEXT("Zen", "Zen_LaunchFailurePromptTitle", "Failed to launch");
+
+						FFormatNamedArguments FormatArguments;
+						FString LogFilePath = FPaths::Combine(InSettings.DataPath, TEXT("logs"), TEXT("zenserver.log"));
+						FPaths::MakePlatformFilename(LogFilePath);
+						FormatArguments.Add(TEXT("LogFilePath"), FText::FromString(LogFilePath));
+						FText ZenLaunchFailurePromptText = FText::Format(NSLOCTEXT("Zen", "Zen_LaunchFailurePromptText", "ZenServer failed to launch. This process will now exit. Please check the ZenServer log file for details:\n{LogFilePath}"), FormatArguments);
+						FPlatformMisc::MessageBoxExt(EAppMsgType::Ok, *ZenLaunchFailurePromptText.ToString(), *ZenLaunchFailurePromptTitle.ToString());
 						FPlatformMisc::RequestExit(true);
 						return false;
 					}
-					DurationPhase = EWaitDurationPhase::Long;
 				}
-
-				if (WaitForZenReadySlowTask.ShouldCancel())
+				// Note that the dialog may not show up when zenserver is needed early in the launch cycle, but this will at least ensure
+				// the splash screen is refreshed with the appropriate text status message.
+				WaitForZenReadySlowTask.MakeDialog(true, false);
+				UE_LOG(LogZenServiceInstance, Display, TEXT("Waiting for ZenServer to be ready..."));
+				DurationPhase = EWaitDurationPhase::Medium;
+			}
+			else if (!FApp::IsUnattended() && ZenWaitDuration > 20.0 && (DurationPhase == EWaitDurationPhase::Medium))
+			{
+				FText ZenLongWaitPromptTitle = NSLOCTEXT("Zen", "Zen_LongWaitPromptTitle", "Wait for ZenServer?");
+				FText ZenLongWaitPromptText = NSLOCTEXT("Zen", "Zen_LongWaitPromptText", "ZenServer is taking a long time to launch. It may be performing maintenance. Keep waiting?");
+				if (FPlatformMisc::MessageBoxExt(EAppMsgType::YesNo, *ZenLongWaitPromptText.ToString(), *ZenLongWaitPromptTitle.ToString()) == EAppReturnType::No)
 				{
 					FPlatformMisc::RequestExit(true);
 					return false;
 				}
-				FPlatformProcess::Sleep(0.1f);
+				DurationPhase = EWaitDurationPhase::Long;
 			}
-		}
 
-		return bIsReady;
-	}
-	else
-	{
-		return false;
+			if (WaitForZenReadySlowTask.ShouldCancel())
+			{
+				FPlatformMisc::RequestExit(true);
+				return false;
+			}
+			FPlatformProcess::Sleep(0.1f);
+		}
 	}
-	return true;
+	return bIsReady;
 }
 
 bool 
