Issue 266313 Static analyzer warnings in ZenServerInterface.cpp

diff --git a/Engine/Source/Developer/Zen/Private/ZenServerInterface.cpp b/Engine/Source/Developer/Zen/Private/ZenServerInterface.cpp
--- a/Engine/Source/Developer/Zen/Private/ZenServerInterface.cpp
+++ b/Engine/Source/Developer/Zen/Private/ZenServerInterface.cpp
@@ -920,10 +920,14 @@ ShutdownRunningService(const TCHAR* ExecutablePath, double MaximumWaitDurationSe
 
 	ZenServerState ServerState;
 	const ZenServerState::ZenServerEntry* Entry = ServerState.LookupByPid(ServicePid);
-	if (!Entry && FPlatformProcess::IsProcRunning(ProcessHandle))
+	if (!Entry)
 	{
-		UE_LOG(LogZenServiceInstance, Warning, TEXT("Can't find server state for running service for executable '%s' (Pid: %u)"), ExecutablePath, ServicePid);
-		return false;
+		if (FPlatformProcess::IsProcRunning(ProcessHandle))
+		{
+			UE_LOG(LogZenServiceInstance, Warning, TEXT("Can't find server state for running service for executable '%s' (Pid: %u)"), ExecutablePath, ServicePid);
+			return false;
+		}
+		return true;
 	}
 	uint16 ProcessPort = Entry->DesiredListenPort;
 	if (!RequestZenShutdownOnPort(ProcessPort) && FPlatformProcess::IsProcRunning(ProcessHandle))
