Issue 266313 Static analyzer warnings in ZenServerInterface.cpp, take 2

diff --git a/Engine/Source/Developer/Zen/Private/ZenServerInterface.cpp b/Engine/Source/Developer/Zen/Private/ZenServerInterface.cpp
--- a/Engine/Source/Developer/Zen/Private/ZenServerInterface.cpp
+++ b/Engine/Source/Developer/Zen/Private/ZenServerInterface.cpp
@@ -962,11 +962,16 @@ ShutdownRunningService(uint16 ProcessPort, double MaximumWaitDurationSeconds = 5
 	UE_LOG(LogZenServiceInstance, Display, TEXT("Waiting for running instance using port %u to shut down"), ProcessPort);
 	ZenServerState ServerState;
 	const ZenServerState::ZenServerEntry* Entry = ServerState.LookupByPort(ProcessPort);
-	if (!Entry && IsZenProcessUsingPort(ProcessPort))
+	if (!Entry)
 	{
-		UE_LOG(LogZenServiceInstance, Warning, TEXT("Can't find server state for running service using port %u"), ProcessPort);
-		return false;
+		if (IsZenProcessUsingPort(ProcessPort))
+		{
+			UE_LOG(LogZenServiceInstance, Warning, TEXT("Can't find server state for running service using port %u"), ProcessPort);
+			return false;
+		}
+		return true;
 	}
+
 	uint32 ServicePid = Entry->Pid;
 	FProcHandle ProcessHandle = FPlatformProcess::OpenProcess(ServicePid);
 	if (!ProcessHandle.IsValid() && IsZenProcessUsingPort(ProcessPort))
