More hardening on error cases when autolaunching zenserver.
Replace some checks causing crashes to errors/warning and exit.

diff --git a/Engine/Source/Developer/Zen/Private/ZenServerInterface.cpp b/Engine/Source/Developer/Zen/Private/ZenServerInterface.cpp
--- a/Engine/Source/Developer/Zen/Private/ZenServerInterface.cpp
+++ b/Engine/Source/Developer/Zen/Private/ZenServerInterface.cpp
@@ -930,13 +930,17 @@ ShutdownRunningService(const TCHAR* ExecutablePath, double MaximumWaitDurationSe
 		return true;
 	}
 	uint16 ProcessPort = Entry->DesiredListenPort;
-	if (!RequestZenShutdownOnPort(ProcessPort) && FPlatformProcess::IsProcRunning(ProcessHandle))
+	if (!RequestZenShutdownOnPort(ProcessPort))
 	{
-		UE_LOG(LogZenServiceInstance, Warning, TEXT("Failed to request shutdown for running service for executable '%s' (Pid: %u, Port: %u)"), ExecutablePath, ServicePid, ProcessPort);
-		return false;
+		if (FPlatformProcess::IsProcRunning(ProcessHandle) || IsZenProcessActive(ExecutablePath, nullptr))
+		{
+			UE_LOG(LogZenServiceInstance, Warning, TEXT("Failed to request shutdown for running service for executable '%s' (Pid: %u, Port: %u)"), ExecutablePath, ServicePid, ProcessPort);
+			return false;
+		}
+		return true;
 	}
 	uint64 ZenShutdownWaitStartTime = FPlatformTime::Cycles64();
-	while (FPlatformProcess::IsProcRunning(ProcessHandle))
+	while (FPlatformProcess::IsProcRunning(ProcessHandle) || IsZenProcessActive(ExecutablePath, nullptr))
 	{
 		double ZenShutdownWaitDuration = FPlatformTime::ToSeconds64(FPlatformTime::Cycles64() - ZenShutdownWaitStartTime);
 		if (ZenShutdownWaitDuration < MaximumWaitDurationSeconds)
@@ -974,20 +978,28 @@ ShutdownRunningService(uint16 ProcessPort, double MaximumWaitDurationSeconds = 5
 
 	uint32 ServicePid = Entry->Pid;
 	FProcHandle ProcessHandle = FPlatformProcess::OpenProcess(ServicePid);
-	if (!ProcessHandle.IsValid() && IsZenProcessUsingPort(ProcessPort))
+	if (!ProcessHandle.IsValid())
 	{
-		UE_LOG(LogZenServiceInstance, Warning, TEXT("Failed to open handle for running service using port %u (Pid: %u)"), ProcessPort, ServicePid);
-		return false;
+		if (IsZenProcessUsingPort(ProcessPort))
+		{
+			UE_LOG(LogZenServiceInstance, Warning, TEXT("Failed to open handle for running service using port %u (Pid: %u)"), ProcessPort, ServicePid);
+			return false;
+		}
+		return true;
 	}
 	ON_SCOPE_EXIT{ FPlatformProcess::CloseProc(ProcessHandle); };
 
-	if (!RequestZenShutdownOnPort(ProcessPort) && FPlatformProcess::IsProcRunning(ProcessHandle))
+	if (!RequestZenShutdownOnPort(ProcessPort))
 	{
-		UE_LOG(LogZenServiceInstance, Warning, TEXT("Failed to request shutdown for running service using port %u (Pid: %u)"), ProcessPort, ServicePid);
-		return false;
+		if (FPlatformProcess::IsProcRunning(ProcessHandle) || IsZenProcessUsingPort(ProcessPort))
+		{
+			UE_LOG(LogZenServiceInstance, Warning, TEXT("Failed to request shutdown for running service using port %u (Pid: %u)"), ProcessPort, ServicePid);
+			return false;
+		}
+		return true;
 	}
 	uint64 ZenShutdownWaitStartTime = FPlatformTime::Cycles64();
-	while (FPlatformProcess::IsProcRunning(ProcessHandle))
+	while (FPlatformProcess::IsProcRunning(ProcessHandle) || IsZenProcessUsingPort(ProcessPort))
 	{
 		double ZenShutdownWaitDuration = FPlatformTime::ToSeconds64(FPlatformTime::Cycles64() - ZenShutdownWaitStartTime);
 		if (ZenShutdownWaitDuration < MaximumWaitDurationSeconds)
@@ -1602,9 +1614,25 @@ FZenServiceInstance::AutoLaunch(const FServiceAutoLaunchSettings& InSettings, FS
 			UE_LOG(LogZenServiceInstance, Warning, TEXT("Failed to shut down running service using port %u"), DesiredPort);
 			return false;
 		}
-		checkf(!IsLockFileLocked(*LockFilePath), TEXT("Data folder is still locked by other process '%s'"), *InSettings.DataPath);
 		checkf(!IsZenProcessUsingPort(DesiredPort), TEXT("Service port is still in use %u"), DesiredPort);
 
+		if (IsLockFileLocked(*LockFilePath))
+		{
+			if (FApp::IsUnattended())
+			{
+				// Just log as there is no one to show a message
+				UE_LOG(LogZenServiceInstance, Warning, TEXT("Failed to auto launch, data folder is still locked by other process '%s'"), *InSettings.DataPath);
+				FPlatformMisc::RequestExit(true);
+				return false;
+			}
+
+			FText ZenLaunchFailurePromptTitle = NSLOCTEXT("Zen", "Zen_LaunchFailurePromptTitle", "Failed to launch");
+			FText ZenLaunchFailurePromptText = FText::Format(NSLOCTEXT("Zen", "Zen_UpdatePromptText", "ZenServer Failed to auto launch, data folder '{0}' is still locked by other process"), FText::FromString(*InSettings.DataPath));
+			FPlatformMisc::MessageBoxExt(EAppMsgType::Ok, *ZenLaunchFailurePromptTitle.ToString(), *ZenLaunchFailurePromptText.ToString());
+			FPlatformMisc::RequestExit(true);
+			return false;
+		}
+
 		FString ParmsWithoutTransients = DetermineCmdLineWithoutTransientComponents(InSettings, DesiredPort);
 		FString TransientParms;
 
@@ -1670,7 +1698,9 @@ FZenServiceInstance::AutoLaunch(const FServiceAutoLaunchSettings& InSettings, FS
 				{
 					if (FApp::IsUnattended())
 					{
-						checkf(false, TEXT("ZenServer did not launch in the expected duration."));
+						// Just log as there is no one to show a message
+						UE_LOG(LogZenServiceInstance, Warning, TEXT("ZenServer did not launch in the expected duration"));
+						return false;
 					}
 					else
 					{
