Resolve locatization Horde issue 267733.
Use unique keys for NSLOCTEXT in ZenServerInstance.cpp

https://horde.devtools.epicgames.com/job/63f84478697ad0eddb32c990?step=9317&issue=267733

diff --git a/Engine/Source/Developer/Zen/Private/ZenServerInterface.cpp b/Engine/Source/Developer/Zen/Private/ZenServerInterface.cpp
--- a/Engine/Source/Developer/Zen/Private/ZenServerInterface.cpp
+++ b/Engine/Source/Developer/Zen/Private/ZenServerInterface.cpp
@@ -1626,8 +1626,8 @@ FZenServiceInstance::AutoLaunch(const FServiceAutoLaunchSettings& InSettings, FS
 				return false;
 			}
 
-			FText ZenLaunchFailurePromptTitle = NSLOCTEXT("Zen", "Zen_LaunchFailurePromptTitle", "Failed to launch");
-			FText ZenLaunchFailurePromptText = FText::Format(NSLOCTEXT("Zen", "Zen_UpdatePromptText", "ZenServer Failed to auto launch, data folder '{0}' is still locked by other process"), FText::FromString(*InSettings.DataPath));
+			FText ZenLaunchFailurePromptTitle = NSLOCTEXT("Zen", "Zen_ShutdownFailurePromptTitle", "Failed to launch");
+			FText ZenLaunchFailurePromptText = FText::Format(NSLOCTEXT("Zen", "Zen_ShutdownFailurePromptText", "ZenServer Failed to auto launch, data folder '{0}' is still locked by other process"), FText::FromString(*InSettings.DataPath));
 			FPlatformMisc::MessageBoxExt(EAppMsgType::Ok, *ZenLaunchFailurePromptTitle.ToString(), *ZenLaunchFailurePromptText.ToString());
 			FPlatformMisc::RequestExit(true);
 			return false;
