Ensure that if ZenServer uses a different effective port than its desired port that Autolauch and shutdown operate with the appropriate effective port.

diff --git a/Engine/Source/Developer/Zen/Private/ZenServerInterface.cpp b/Engine/Source/Developer/Zen/Private/ZenServerInterface.cpp
--- a/Engine/Source/Developer/Zen/Private/ZenServerInterface.cpp
+++ b/Engine/Source/Developer/Zen/Private/ZenServerInterface.cpp
@@ -931,7 +931,7 @@ ShutdownRunningService(const TCHAR* ExecutablePath, double MaximumWaitDurationSe
 		}
 		return true;
 	}
-	uint16 ProcessPort = Entry->DesiredListenPort;
+	uint16 ProcessPort = Entry->EffectiveListenPort;
 	if (!RequestZenShutdownOnPort(ProcessPort))
 	{
 		if (FPlatformProcess::IsProcRunning(ProcessHandle) || IsZenProcessActive(ExecutablePath, nullptr))
@@ -1584,6 +1584,7 @@ FZenServiceInstance::AutoLaunch(const FServiceAutoLaunchSettings& InSettings, FS
 	FString WorkingDirectory = FPaths::GetPath(ExecutablePath);
 
 	bool bReUsingExistingInstance = false;
+	uint16 ShutdownPort = DesiredPort;
 
 	// When limiting process lifetime, always re-launch to add sponsor process IDs.
 	// When not limiting process lifetime, only launch if the process is not already live.
@@ -1592,9 +1593,11 @@ FZenServiceInstance::AutoLaunch(const FServiceAutoLaunchSettings& InSettings, FS
 		uint16 CurrentPort;
 		if (IsZenProcessUsingDataDir(*InSettings.DataPath, &CurrentPort) && CurrentPort != 0)
 		{
+			ShutdownPort = CurrentPort;
+
 			FZenLocalServiceRunContext DesiredRunContext;
 			DesiredRunContext.Executable = ExecutablePath;
-			DesiredRunContext.CommandlineArguments = DetermineCmdLineWithoutTransientComponents(InSettings, CurrentPort);
+			DesiredRunContext.CommandlineArguments = DetermineCmdLineWithoutTransientComponents(InSettings, DesiredPort);
 			DesiredRunContext.WorkingDirectory = WorkingDirectory;
 			DesiredRunContext.DataPath = InSettings.DataPath;
 			DesiredRunContext.bShowConsole = InSettings.bShowConsole;
@@ -1611,12 +1614,12 @@ FZenServiceInstance::AutoLaunch(const FServiceAutoLaunchSettings& InSettings, FS
 
 	if (!bReUsingExistingInstance)
 	{
-		if (!ShutdownRunningService(DesiredPort))
+		if (!ShutdownRunningService(ShutdownPort))
 		{
-			UE_LOG(LogZenServiceInstance, Warning, TEXT("Failed to shut down running service using port %u"), DesiredPort);
+			UE_LOG(LogZenServiceInstance, Warning, TEXT("Failed to shut down running service using port %u"), ShutdownPort);
 			return false;
 		}
-		checkf(!IsZenProcessUsingPort(DesiredPort), TEXT("Service port is still in use %u"), DesiredPort);
+		checkf(!IsZenProcessUsingPort(ShutdownPort), TEXT("Service port is still in use %u"), ShutdownPort);
 
 		if (IsLockFileLocked(*LockFilePath))
 		{
