Reworked zen server instance start/stop logic at autolaunch to handle Zen:LimitProcessLifetime flag properly.
Handle scenario where the data dir is locked by a process not running locally by informing user and exit gracefully.
Reuse running zen server instance if all but current port is the same and adjust port to running instance.

diff --git a/Engine/Source/Developer/Zen/Private/ZenServerInterface.cpp b/Engine/Source/Developer/Zen/Private/ZenServerInterface.cpp
--- a/Engine/Source/Developer/Zen/Private/ZenServerInterface.cpp
+++ b/Engine/Source/Developer/Zen/Private/ZenServerInterface.cpp
@@ -840,10 +840,9 @@ IsZenProcessActive(const TCHAR* ExecutablePath, uint32* OutPid)
 }
 
 static bool
-IsZenProcessUsingDataDir(const TCHAR* DataPath, uint16* OutPort)
+IsZenProcessUsingDataDir(const TCHAR* LockFilePath, uint16* OutPort)
 {
-	const FString LockFilePath = FPaths::Combine(DataPath, TEXT(".lock"));
-	if (IsLockFileLocked(*LockFilePath, true))
+	if (IsLockFileLocked(LockFilePath, true))
 	{
 		if (OutPort)
 		{
@@ -1150,7 +1149,8 @@ FZenLocalServiceRunContext::WriteToJsonFile(const TCHAR* Filename) const
 bool
 IsLocalServiceRunning(const TCHAR* DataPath, uint16* OutPort)
 {
-	return IsZenProcessUsingDataDir(DataPath, OutPort);
+	const FString LockFilePath = FPaths::Combine(DataPath, TEXT(".lock"));
+	return IsZenProcessUsingDataDir(*LockFilePath, OutPort);
 }
 
 FProcHandle
@@ -1385,7 +1385,7 @@ FZenServiceInstance::IsServiceReady()
 		
 		if (Result == Zen::FZenHttpRequest::Result::Success && Zen::IsSuccessCode(Request.GetResponseCode()))
 		{
-			UE_LOG(LogZenServiceInstance, Display, TEXT("ZenServer HTTP service status: %s."), *Request.GetResponseAsString());
+			UE_LOG(LogZenServiceInstance, Display, TEXT("ZenServer HTTP service at %s status: %s."), ZenDomain.ToString(), *Request.GetResponseAsString());
 			return true;
 		}
 		else
@@ -1497,8 +1497,8 @@ FZenServiceInstance::Initialize()
 	URL = WriteToString<64>(TEXT("http://"), HostName, TEXT(":"), Port, TEXT("/"));
 }
 
-void
-FZenServiceInstance::PromptUserToStopRunningServerInstance(const FString& ServerFilePath)
+static void
+PromptUserToStopRunningServerInstanceForUpdate(const FString& ServerFilePath)
 {
 	if (FApp::IsUnattended())
 	{
@@ -1512,6 +1512,38 @@ FZenServiceInstance::PromptUserToStopRunningServerInstance(const FString& Server
 	FPlatformMisc::MessageBoxExt(EAppMsgType::Ok, *ZenUpdatePromptText.ToString(), *ZenUpdatePromptTitle.ToString());
 }
 
+static void
+PromptUserOfLockedDataFolder(const FString& DataPath)
+{
+	if (FApp::IsUnattended())
+	{
+		// Just log as there is no one to show a message
+		UE_LOG(LogZenServiceInstance, Warning, TEXT("ZenServer Failed to auto launch, an unknown process is locking the data folder '%s'"), *DataPath);
+		return;
+	}
+
+	FText ZenLaunchFailurePromptTitle = NSLOCTEXT("Zen", "Zen_NonLocalProcessUsesDataDirPromptTitle", "Failed to launch");
+	FText ZenLaunchFailurePromptText = FText::Format(NSLOCTEXT("Zen", "Zen_NonLocalProcessUsesDataDirPromptText", "ZenServer Failed to auto launch, an unknown process is locking the data folder '{0}'"), FText::FromString(DataPath));
+	FPlatformMisc::MessageBoxExt(EAppMsgType::Ok, *ZenLaunchFailurePromptTitle.ToString(), *ZenLaunchFailurePromptText.ToString());
+}
+
+static void
+PromptUserOfFailedShutDownOfExistingProcess(uint16 Port)
+{
+	if (FApp::IsUnattended())
+	{
+		// Just log as there is no one to show a message
+		UE_LOG(LogZenServiceInstance, Warning, TEXT("ZenServer Failed to auto launch, failed to shut down currently running service using port %u"), Port);
+		FPlatformMisc::RequestExit(true);
+		return;
+	}
+
+	FText ZenLaunchFailurePromptTitle = NSLOCTEXT("Zen", "Zen_ShutdownFailurePromptTitle", "Failed to launch");
+	FText ZenLaunchFailurePromptText = FText::Format(NSLOCTEXT("Zen", "Zen_ShutdownFailurePromptText", "ZenServer Failed to auto launch, failed to shut down currently running service using port '{0}'"), Port);
+	FPlatformMisc::MessageBoxExt(EAppMsgType::Ok, *ZenLaunchFailurePromptTitle.ToString(), *ZenLaunchFailurePromptText.ToString());
+}
+
+
 FString
 FZenServiceInstance::ConditionalUpdateLocalInstall()
 {
@@ -1530,20 +1562,20 @@ FZenServiceInstance::ConditionalUpdateLocalInstall()
 		UE_LOG(LogZenServiceInstance, Display, TEXT("Installing service from '%s' to '%s'"), *InTreeServicePath, *InstallServicePath);
 		if (!ShutdownRunningService(*InstallServicePath))
 		{
-			PromptUserToStopRunningServerInstance(InstallServicePath);
+			PromptUserToStopRunningServerInstanceForUpdate(InstallServicePath);
 			return FString();
 		}
 
 		// Even after waiting for the process to shut down we have a tolerance for failure when overwriting the target files
 		if (!AttemptFileCopyWithRetries(*InstallServicePath, *InTreeServicePath, 5.0))
 		{
-			PromptUserToStopRunningServerInstance(InstallServicePath);
+			PromptUserToStopRunningServerInstanceForUpdate(InstallServicePath);
 			return FString();
 		}
 
 		if (!AttemptFileCopyWithRetries(*InstallUtilityPath, *InTreeUtilityPath, 5.0))
 		{
-			PromptUserToStopRunningServerInstance(InstallServicePath);
+			PromptUserToStopRunningServerInstanceForUpdate(InstallServicePath);
 			return FString();
 		}
 
@@ -1583,61 +1615,79 @@ FZenServiceInstance::AutoLaunch(const FServiceAutoLaunchSettings& InSettings, FS
 
 	FString WorkingDirectory = FPaths::GetPath(ExecutablePath);
 
-	bool bReUsingExistingInstance = false;
-	uint16 ShutdownPort = DesiredPort;
+	uint16 RunningProcessPort = 0;
+	uint64 ZenWaitForRunningProcessReadyStartTime = FPlatformTime::Cycles64();
+	while (IsZenProcessUsingDataDir(*LockFilePath, &RunningProcessPort) && RunningProcessPort == 0)
+	{
+		// Server is starting up, wait for it to get ready
+		double ZenWaitDuration = FPlatformTime::ToSeconds64(FPlatformTime::Cycles64() - ZenWaitForRunningProcessReadyStartTime);
+		if (ZenWaitDuration > 5.0)
+		{
+			break;
+		}
+		FPlatformProcess::Sleep(0.1f);
+	}
 
-	// When limiting process lifetime, always re-launch to add sponsor process IDs.
-	// When not limiting process lifetime, only launch if the process is not already live.
-	if (!InSettings.bLimitProcessLifetime)
+	uint16 ShutdownPort = 0;
+	bool bLaunchNewInstance = true;
+
+	if (RunningProcessPort != 0)
 	{
-		uint16 CurrentPort;
-		if (IsZenProcessUsingDataDir(*InSettings.DataPath, &CurrentPort) && CurrentPort != 0)
+		if (!IsZenProcessUsingPort(RunningProcessPort))
 		{
-			ShutdownPort = CurrentPort;
+			PromptUserOfLockedDataFolder(*InSettings.DataPath);
+			FPlatformMisc::RequestExit(true);
+			return false;
+		}
 
-			FZenLocalServiceRunContext DesiredRunContext;
-			DesiredRunContext.Executable = ExecutablePath;
-			DesiredRunContext.CommandlineArguments = DetermineCmdLineWithoutTransientComponents(InSettings, DesiredPort);
-			DesiredRunContext.WorkingDirectory = WorkingDirectory;
-			DesiredRunContext.DataPath = InSettings.DataPath;
-			DesiredRunContext.bShowConsole = InSettings.bShowConsole;
+		FZenLocalServiceRunContext DesiredRunContext;
+		DesiredRunContext.Executable = ExecutablePath;
+		DesiredRunContext.CommandlineArguments = DetermineCmdLineWithoutTransientComponents(InSettings, RunningProcessPort);
+		DesiredRunContext.WorkingDirectory = WorkingDirectory;
+		DesiredRunContext.DataPath = InSettings.DataPath;
+		DesiredRunContext.bShowConsole = InSettings.bShowConsole;
 
-			FZenLocalServiceRunContext CurrentRunContext;
-			if (CurrentRunContext.ReadFromJsonFile(*ExecutionContextFilePath) && (DesiredRunContext == CurrentRunContext))
-			{
-				UE_LOG(LogZenServiceInstance, Log, TEXT("Reusing existing instance running on port %u"), CurrentPort);
-				DesiredPort = CurrentPort;
-				bReUsingExistingInstance = true;
-			}
+		FZenLocalServiceRunContext CurrentRunContext;
+		if (CurrentRunContext.ReadFromJsonFile(*ExecutionContextFilePath) && (DesiredRunContext == CurrentRunContext))
+		{
+			UE_LOG(LogZenServiceInstance, Log, TEXT("Found existing instance running on port %u matching our settings"), RunningProcessPort);
+			DesiredPort = RunningProcessPort;
+			bLaunchNewInstance = false;
 		}
+		else
+		{
+			UE_LOG(LogZenServiceInstance, Log, TEXT("Found existing instance running on port %u with different run context, will attempt shut down"), RunningProcessPort);
+			ShutdownPort = RunningProcessPort;
+		}
+	}
+	else if (IsZenProcessUsingPort(DesiredPort))
+	{
+		UE_LOG(LogZenServiceInstance, Log, TEXT("Found existing instance running on port %u with different data directory, will attempt shut down"), DesiredPort);
+		ShutdownPort = DesiredPort;
 	}
 
-	if (!bReUsingExistingInstance)
+	if (ShutdownPort != 0)
 	{
 		if (!ShutdownRunningService(ShutdownPort))
 		{
-			UE_LOG(LogZenServiceInstance, Warning, TEXT("Failed to shut down running service using port %u"), ShutdownPort);
+			PromptUserOfFailedShutDownOfExistingProcess(ShutdownPort);
+			FPlatformMisc::RequestExit(true);
 			return false;
 		}
 		checkf(!IsZenProcessUsingPort(ShutdownPort), TEXT("Service port is still in use %u"), ShutdownPort);
 
 		if (IsLockFileLocked(*LockFilePath))
 		{
-			if (FApp::IsUnattended())
-			{
-				// Just log as there is no one to show a message
-				UE_LOG(LogZenServiceInstance, Warning, TEXT("Failed to auto launch, data folder is still locked by other process '%s'"), *InSettings.DataPath);
-				FPlatformMisc::RequestExit(true);
-				return false;
-			}
-
-			FText ZenLaunchFailurePromptTitle = NSLOCTEXT("Zen", "Zen_ShutdownFailurePromptTitle", "Failed to launch");
-			FText ZenLaunchFailurePromptText = FText::Format(NSLOCTEXT("Zen", "Zen_ShutdownFailurePromptText", "ZenServer Failed to auto launch, data folder '{0}' is still locked by other process"), FText::FromString(*InSettings.DataPath));
-			FPlatformMisc::MessageBoxExt(EAppMsgType::Ok, *ZenLaunchFailurePromptTitle.ToString(), *ZenLaunchFailurePromptText.ToString());
+			PromptUserOfLockedDataFolder(*InSettings.DataPath);
 			FPlatformMisc::RequestExit(true);
 			return false;
 		}
+	}
 
+	// When limiting process lifetime, always re-launch to add sponsor process IDs.
+	// When not limiting process lifetime, only launch if the process is not already live.
+	if (bLaunchNewInstance || InSettings.bLimitProcessLifetime)
+	{
 		FString ParmsWithoutTransients = DetermineCmdLineWithoutTransientComponents(InSettings, DesiredPort);
 		FString TransientParms;
 
@@ -1654,7 +1704,12 @@ FZenServiceInstance::AutoLaunch(const FServiceAutoLaunchSettings& InSettings, FS
 		EffectiveRunContext.bShowConsole = InSettings.bShowConsole;
 
 		FProcHandle Proc = StartLocalService(EffectiveRunContext, TransientParms.IsEmpty() ? nullptr : *TransientParms);
-		EffectiveRunContext.WriteToJsonFile(*ExecutionContextFilePath);
+
+		// Don't write run context unless we actually started a new instance
+		if (bLaunchNewInstance)
+		{
+			EffectiveRunContext.WriteToJsonFile(*ExecutionContextFilePath);
+		}
 
 		if (!Proc.IsValid())
 		{
diff --git a/Engine/Source/Developer/Zen/Public/Experimental/ZenServerInterface.h b/Engine/Source/Developer/Zen/Public/Experimental/ZenServerInterface.h
--- a/Engine/Source/Developer/Zen/Public/Experimental/ZenServerInterface.h
+++ b/Engine/Source/Developer/Zen/Public/Experimental/ZenServerInterface.h
@@ -195,7 +195,6 @@ public:
 private:
 
 	void Initialize();
-	void PromptUserToStopRunningServerInstance(const FString& ServerFilePath);
 	FString ConditionalUpdateLocalInstall();
 	static bool AutoLaunch(const FServiceAutoLaunchSettings& InSettings, FString&& ExecutablePath, FString& OutHostName, uint16& OutPort);
 	
