Refactor zen autolaunch
- Properly identify running instance with matching configuration when using a remapped effective port
- Properly handle shut down of configuration mismatch when using a remapped effective port
- Use shared ZenServerState to add sponsor processes instead of re-launching zenserver multiple times (zenserver is broken in that scenario)
- Clean up naming of functions and parameters to make intention extra clear
- Fix mixup of title vs text fields in some text boxes
- Fix shutdown of zenserver process if we can't get a run context in TryRecovery()

diff --git a/Engine/Source/Developer/Zen/Private/ZenServerInterface.cpp b/Engine/Source/Developer/Zen/Private/ZenServerInterface.cpp
--- a/Engine/Source/Developer/Zen/Private/ZenServerInterface.cpp
+++ b/Engine/Source/Developer/Zen/Private/ZenServerInterface.cpp
@@ -56,7 +56,7 @@ DEFINE_LOG_CATEGORY_STATIC(LogZenServiceInstance, Log, All);
 class ZenServerState
 {
 public:
-	ZenServerState();
+	ZenServerState(bool ReadOnly = true);
 	~ZenServerState();
 
 	struct ZenServerEntry
@@ -77,29 +77,35 @@ public:
 			kShutdownPlease = 1 << 0,
 			kIsReady = 1 << 1,
 		};
+
+		bool AddSponsorProcess(uint32_t PidToAdd);
 	};
 	static_assert(sizeof(ZenServerEntry) == 64);
 
-	const ZenServerEntry* LookupByPort(int DesiredListenPort);
-	const ZenServerEntry* LookupByPid(uint32 Pid);
+	ZenServerEntry* LookupByDesiredListenPort(int DesiredListenPort);
+	const ZenServerEntry* LookupByDesiredListenPort(int DesiredListenPort) const;
+	const ZenServerEntry* LookupByEffectiveListenPort(int DesiredListenPort) const;
+	const ZenServerEntry* LookupByPid(uint32 Pid) const;
 
 private:
 	void* m_hMapFile = nullptr;
-	const ZenServerEntry* m_Data = nullptr;
+	ZenServerEntry* m_Data = nullptr;
 	int				m_MaxEntryCount = 65536 / sizeof(ZenServerEntry);
+	bool			m_IsReadOnly = true;
 };
 
-ZenServerState::ZenServerState()
+ZenServerState::ZenServerState(bool ReadOnly)
 	: m_hMapFile(nullptr)
 	, m_Data(nullptr)
 {
 	size_t MapSize = m_MaxEntryCount * sizeof(ZenServerEntry);
 
 #if PLATFORM_WINDOWS
-	HANDLE hMap = OpenFileMapping(FILE_MAP_READ, 0, L"Global\\ZenMap");
+	DWORD DesiredAccess = ReadOnly ? FILE_MAP_READ : (FILE_MAP_READ | FILE_MAP_WRITE);
+	HANDLE hMap = OpenFileMapping(DesiredAccess, 0, L"Global\\ZenMap");
 	if (hMap == NULL)
 	{
-		hMap = OpenFileMapping(FILE_MAP_READ, 0, L"Local\\ZenMap");
+		hMap = OpenFileMapping(DesiredAccess, 0, L"Local\\ZenMap");
 	}
 
 	if (hMap == NULL)
@@ -107,8 +113,8 @@ ZenServerState::ZenServerState()
 		return;
 	}
 
-	const void* pBuf = MapViewOfFile(hMap,		   // handle to map object
-		FILE_MAP_READ,  // read permission
+	void* pBuf = MapViewOfFile(hMap,		   // handle to map object
+		DesiredAccess,  // read permission
 		0,			   // offset high
 		0,			   // offset low
 		MapSize);
@@ -126,7 +132,8 @@ ZenServerState::ZenServerState()
 	}
 	void* hMap = (void*)intptr_t(Fd);
 
-	const void* pBuf = mmap(nullptr, MapSize, PROT_READ, MAP_PRIVATE, Fd, 0);
+	int Prot = ReadOnly ? (PROT_WRITE | PROT_READ) : PROT_READ;
+	void* pBuf = mmap(nullptr, MapSize, Prot, MAP_PRIVATE, Fd, 0);
 	if (pBuf == MAP_FAILED)
 	{
 		close(Fd);
@@ -136,8 +143,9 @@ ZenServerState::ZenServerState()
 
 #if PLATFORM_WINDOWS || PLATFORM_UNIX || PLATFORM_MAC
 	m_hMapFile = hMap;
-	m_Data = reinterpret_cast<const ZenServerEntry*>(pBuf);
+	m_Data = reinterpret_cast<ZenServerEntry*>(pBuf);
 #endif
+	m_IsReadOnly = ReadOnly;
 }
 
 ZenServerState::~ZenServerState()
@@ -165,7 +173,45 @@ ZenServerState::~ZenServerState()
 	m_Data = nullptr;
 }
 
-const ZenServerState::ZenServerEntry* ZenServerState::LookupByPort(int DesiredListenPort)
+const ZenServerState::ZenServerEntry* ZenServerState::LookupByDesiredListenPort(int Port) const
+{
+	if (m_Data == nullptr)
+	{
+		return nullptr;
+	}
+
+	for (int i = 0; i < m_MaxEntryCount; ++i)
+	{
+		if (m_Data[i].DesiredListenPort == Port)
+		{
+			return &m_Data[i];
+		}
+	}
+
+	return nullptr;
+}
+
+ZenServerState::ZenServerEntry* ZenServerState::LookupByDesiredListenPort(int Port)
+{
+	check(!m_IsReadOnly);
+
+	if (m_Data == nullptr)
+	{
+		return nullptr;
+	}
+
+	for (int i = 0; i < m_MaxEntryCount; ++i)
+	{
+		if (m_Data[i].DesiredListenPort == Port)
+		{
+			return &m_Data[i];
+		}
+	}
+
+	return nullptr;
+}
+
+const ZenServerState::ZenServerEntry* ZenServerState::LookupByEffectiveListenPort(int Port) const
 {
 	if (m_Data == nullptr)
 	{
@@ -174,7 +220,7 @@ const ZenServerState::ZenServerEntry* ZenServerState::LookupByPort(int DesiredLi
 
 	for (int i = 0; i < m_MaxEntryCount; ++i)
 	{
-		if (m_Data[i].DesiredListenPort == DesiredListenPort)
+		if (m_Data[i].EffectiveListenPort == Port)
 		{
 			return &m_Data[i];
 		}
@@ -183,7 +229,7 @@ const ZenServerState::ZenServerEntry* ZenServerState::LookupByPort(int DesiredLi
 	return nullptr;
 }
 
-const ZenServerState::ZenServerEntry* ZenServerState::LookupByPid(uint32 Pid)
+const ZenServerState::ZenServerEntry* ZenServerState::LookupByPid(uint32 Pid) const
 {
 	if (m_Data == nullptr)
 	{
@@ -201,6 +247,30 @@ const ZenServerState::ZenServerEntry* ZenServerState::LookupByPid(uint32 Pid)
 	return nullptr;
 }
 
+bool
+ZenServerState::ZenServerEntry::AddSponsorProcess(uint32_t PidToAdd)
+{
+	for (std::atomic<uint32_t>& PidEntry : SponsorPids)
+	{
+		if (PidEntry.load(std::memory_order_relaxed) == 0)
+		{
+			uint32_t Expected = 0;
+			if (PidEntry.compare_exchange_strong(Expected, PidToAdd))
+			{
+				// Success!
+				return true;
+			}
+		}
+		else if (PidEntry.load(std::memory_order_relaxed) == PidToAdd)
+		{
+			// Success, the because pid is already in the list
+			return true;
+		}
+	}
+
+	return false;
+}
+
 static FZenVersion
 GetZenVersion(const FString& UtilityPath, const FString& ServicePath, const FString& VersionCachePath)
 {
@@ -655,8 +725,34 @@ FServiceSettings::TryApplyAutoLaunchOverride()
 
 uint16 FZenServiceInstance::AutoLaunchedPort = 0;
 
-static bool
-ReadCbLockFile(FStringView FileName, FCbObject& OutLockObject)
+struct LockFileData
+{
+	uint32 ProcessId = 0;
+	FString DataDir;
+	uint16 EffectivePort = 0;
+	FCbObjectId SessionId;
+	bool IsReady = false;
+	bool IsValid = false;
+};
+
+static LockFileData ReadLockData(FUniqueBuffer&& FileBytes)
+{
+	if (ValidateCompactBinary(FileBytes, ECbValidateMode::Default) == ECbValidateError::None)
+	{
+		FCbObject LockObject(FileBytes.MoveToShared());
+		int32 ProcessId = LockObject["pid"].AsInt32();
+		FUtf8StringView DataDir = LockObject["data"].AsString();
+		int32 EffectivePort = LockObject["port"].AsInt32();
+		FCbObjectId SessionId = LockObject["session_id"].AsObjectId();
+		bool IsReady = LockObject["ready"].AsBool();
+		bool IsValid = ProcessId > 0 && EffectivePort > 0 && EffectivePort <= 0xffff;
+		return LockFileData{static_cast<uint32>(ProcessId), FString(DataDir), static_cast<uint16>(EffectivePort), SessionId, IsReady, IsValid};
+	}
+	return LockFileData{ 0, {}, 0, {}, false};
+}
+
+static LockFileData
+ReadCbLockFile(FStringView FileName)
 {
 #if PLATFORM_WINDOWS
 	// Windows specific lock reading path
@@ -692,22 +788,18 @@ ReadCbLockFile(FStringView FileName, FCbObject& OutLockObject)
 			DWORD ReadBytes = 0;
 			if (ReadFile(Handle, FileBytes.GetData(), FileSize32, &ReadBytes, NULL) && (ReadBytes == FileSize32))
 			{
-				if (ValidateCompactBinary(FileBytes, ECbValidateMode::Default) == ECbValidateError::None)
-				{
-					OutLockObject = FCbObject(FileBytes.MoveToShared());
-					return true;
-				}
+				return ReadLockData(std::move(FileBytes));
 			}
 		}
 	}
-	return false;
+	return {};
 #elif PLATFORM_UNIX || PLATFORM_MAC
 	TAnsiStringBuilder<256> LockFilePath;
 	LockFilePath << FileName;
 	int32 Fd = open(LockFilePath.ToString(), O_RDONLY);
 	if (Fd < 0)
 	{
-		return false;
+		return {};
 	}
 
 	// If we can claim the lock then it's an orphaned lock file and should be
@@ -719,31 +811,26 @@ ReadCbLockFile(FStringView FileName, FCbObject& OutLockObject)
 		unlink(LockFilePath.ToString());
 		flock(Fd, LOCK_UN);
 		close(Fd);
-		return false;
+		return {};
 	}
 
 	if (errno != EWOULDBLOCK && errno != EAGAIN)
 	{
-		return false;
+		return {};
 	}
 
 	struct stat Stat;
 	fstat(Fd, &Stat);
 	uint64 FileSize = uint64(Stat.st_size);
 
-	bool bSuccess = false;
 	FUniqueBuffer FileBytes = FUniqueBuffer::Alloc(FileSize);
 	if (read(Fd, FileBytes.GetData(), FileSize) == FileSize)
 	{
-		if (ValidateCompactBinary(FileBytes, ECbValidateMode::Default) == ECbValidateError::None)
-		{
-			OutLockObject = FCbObject(FileBytes.MoveToShared());
-			bSuccess = true;
-		}
+		return ReadLockData(std::move(FileBytes));
 	}
 
 	close(Fd);
-	return bSuccess;
+	return {};
 #endif
 }
 
@@ -784,10 +871,10 @@ IsLockFileLocked(const TCHAR* FileName, bool bAttemptCleanUp=false)
 }
 
 static bool
-IsZenProcessUsingPort(uint16 Port)
+IsZenProcessUsingEffectivePort(uint16 EffectiveListenPort)
 {
 #if PLATFORM_WINDOWS
-	HANDLE Handle = OpenEventW(READ_CONTROL, false, *WriteToWideString<64>(WIDETEXT("Zen_"), Port, WIDETEXT("_Shutdown")));
+	HANDLE Handle = OpenEventW(READ_CONTROL, false, *WriteToWideString<64>(WIDETEXT("Zen_"), EffectiveListenPort, WIDETEXT("_Shutdown")));
 	if (Handle != NULL)
 	{
 		ON_SCOPE_EXIT{ CloseHandle(Handle); };
@@ -796,7 +883,7 @@ IsZenProcessUsingPort(uint16 Port)
 	return false;
 #elif PLATFORM_UNIX || PLATFORM_MAC
 	TAnsiStringBuilder<64> EventPath;
-	EventPath << "/tmp/Zen_" << Port << "_Shutdown";
+	EventPath << "/tmp/Zen_" << EffectiveListenPort << "_Shutdown";
 
 	key_t IpcKey = ftok(EventPath.ToString(), 1);
 	if (IpcKey < 0)
@@ -840,35 +927,25 @@ IsZenProcessActive(const TCHAR* ExecutablePath, uint32* OutPid)
 }
 
 static bool
-IsZenProcessUsingDataDir(const TCHAR* LockFilePath, uint16* OutPort)
+IsZenProcessUsingDataDir(const TCHAR* LockFilePath, LockFileData* OutLockFileData)
 {
 	if (IsLockFileLocked(LockFilePath, true))
 	{
-		if (OutPort)
+		if (OutLockFileData)
 		{
 			// If an instance is running with this data path, check if we can use it and what port it is on
-			*OutPort = 0;
-			FCbObject LockObject;
-			if (ReadCbLockFile(LockFilePath, LockObject))
-			{
-				bool bIsReady = LockObject["ready"].AsBool();
-				if (bIsReady)
-				{
-					*OutPort = LockObject["port"].AsUInt16();
-				}
-			}
+			*OutLockFileData = ReadCbLockFile(LockFilePath);
 		}
-
 		return true;
 	}
 	return false;
 }
 
 static bool
-RequestZenShutdownOnPort(uint16 Port)
+RequestZenShutdownOnEffectivePort(uint16 EffectiveListenPort)
 {
 #if PLATFORM_WINDOWS
-	HANDLE Handle = OpenEventW(EVENT_MODIFY_STATE, false, *WriteToWideString<64>(WIDETEXT("Zen_"), Port, WIDETEXT("_Shutdown")));
+	HANDLE Handle = OpenEventW(EVENT_MODIFY_STATE, false, *WriteToWideString<64>(WIDETEXT("Zen_"), EffectiveListenPort, WIDETEXT("_Shutdown")));
 	if (Handle != NULL)
 	{
 		ON_SCOPE_EXIT{ CloseHandle(Handle); };
@@ -878,7 +955,7 @@ RequestZenShutdownOnPort(uint16 Port)
 	return false;
 #elif PLATFORM_UNIX || PLATFORM_MAC
 	TAnsiStringBuilder<64> EventPath;
-	EventPath << "/tmp/Zen_" << Port << "_Shutdown";
+	EventPath << "/tmp/Zen_" << EffectiveListenPort << "_Shutdown";
 
 	key_t IpcKey = ftok(EventPath.ToString(), 1);
 	if (IpcKey < 0)
@@ -902,7 +979,7 @@ RequestZenShutdownOnPort(uint16 Port)
 }
 
 static bool
-ShutdownRunningService(const TCHAR* ExecutablePath, double MaximumWaitDurationSeconds = 5.0)
+ShutdownRunningServiceUsingExecutablePath(const TCHAR* ExecutablePath, double MaximumWaitDurationSeconds = 5.0)
 {
 	uint32 ServicePid;
 	if (!IsZenProcessActive(ExecutablePath, &ServicePid))
@@ -919,7 +996,7 @@ ShutdownRunningService(const TCHAR* ExecutablePath, double MaximumWaitDurationSe
 	}
 	ON_SCOPE_EXIT{ FPlatformProcess::CloseProc(ProcessHandle); };
 
-	ZenServerState ServerState;
+	const ZenServerState ServerState;
 	const ZenServerState::ZenServerEntry* Entry = ServerState.LookupByPid(ServicePid);
 	if (!Entry)
 	{
@@ -930,12 +1007,12 @@ ShutdownRunningService(const TCHAR* ExecutablePath, double MaximumWaitDurationSe
 		}
 		return true;
 	}
-	uint16 ProcessPort = Entry->EffectiveListenPort;
-	if (!RequestZenShutdownOnPort(ProcessPort))
+	uint16 EffectiveListenPort = Entry->EffectiveListenPort;
+	if (!RequestZenShutdownOnEffectivePort(EffectiveListenPort))
 	{
 		if (FPlatformProcess::IsProcRunning(ProcessHandle) || IsZenProcessActive(ExecutablePath, nullptr))
 		{
-			UE_LOG(LogZenServiceInstance, Warning, TEXT("Failed to request shutdown for running service for executable '%s' (Pid: %u, Port: %u)"), ExecutablePath, ServicePid, ProcessPort);
+			UE_LOG(LogZenServiceInstance, Warning, TEXT("Failed to request shutdown for running service for executable '%s' (Pid: %u, Port: %u)"), ExecutablePath, ServicePid, EffectiveListenPort);
 			return false;
 		}
 		return true;
@@ -950,7 +1027,7 @@ ShutdownRunningService(const TCHAR* ExecutablePath, double MaximumWaitDurationSe
 		}
 		else
 		{
-			UE_LOG(LogZenServiceInstance, Warning, TEXT("Timed out waiting for shutdown of running service for executable '%s' (Pid: %u, Port: %u)"), ExecutablePath, ServicePid, ProcessPort);
+			UE_LOG(LogZenServiceInstance, Warning, TEXT("Timed out waiting for shutdown of running service for executable '%s' (Pid: %u, Port: %u)"), ExecutablePath, ServicePid, EffectiveListenPort);
 			return false;
 		}
 	}
@@ -958,20 +1035,20 @@ ShutdownRunningService(const TCHAR* ExecutablePath, double MaximumWaitDurationSe
 }
 
 static bool
-ShutdownRunningService(uint16 ProcessPort, double MaximumWaitDurationSeconds = 5.0)
+ShutdownRunningServiceUsingEffectivePort(uint16 EffectiveListenPort, double MaximumWaitDurationSeconds = 5.0)
 {
-	if (!IsZenProcessUsingPort(ProcessPort))
+	if (!IsZenProcessUsingEffectivePort(EffectiveListenPort))
 	{
 		return true;
 	}
-	UE_LOG(LogZenServiceInstance, Display, TEXT("Waiting for running instance using port %u to shut down"), ProcessPort);
-	ZenServerState ServerState;
-	const ZenServerState::ZenServerEntry* Entry = ServerState.LookupByPort(ProcessPort);
+	UE_LOG(LogZenServiceInstance, Display, TEXT("Waiting for running instance using port %u to shut down"), EffectiveListenPort);
+	const ZenServerState ServerState;
+	const ZenServerState::ZenServerEntry* Entry = ServerState.LookupByEffectiveListenPort(EffectiveListenPort);
 	if (!Entry)
 	{
-		if (IsZenProcessUsingPort(ProcessPort))
+		if (IsZenProcessUsingEffectivePort(EffectiveListenPort))
 		{
-			UE_LOG(LogZenServiceInstance, Warning, TEXT("Can't find server state for running service using port %u"), ProcessPort);
+			UE_LOG(LogZenServiceInstance, Warning, TEXT("Can't find server state for running service using port %u"), EffectiveListenPort);
 			return false;
 		}
 		return true;
@@ -981,26 +1058,26 @@ ShutdownRunningService(uint16 ProcessPort, double MaximumWaitDurationSeconds = 5
 	FProcHandle ProcessHandle = FPlatformProcess::OpenProcess(ServicePid);
 	if (!ProcessHandle.IsValid())
 	{
-		if (IsZenProcessUsingPort(ProcessPort))
+		if (IsZenProcessUsingEffectivePort(EffectiveListenPort))
 		{
-			UE_LOG(LogZenServiceInstance, Warning, TEXT("Failed to open handle for running service using port %u (Pid: %u)"), ProcessPort, ServicePid);
+			UE_LOG(LogZenServiceInstance, Warning, TEXT("Failed to open handle for running service using port %u (Pid: %u)"), EffectiveListenPort, ServicePid);
 			return false;
 		}
 		return true;
 	}
 	ON_SCOPE_EXIT{ FPlatformProcess::CloseProc(ProcessHandle); };
 
-	if (!RequestZenShutdownOnPort(ProcessPort))
+	if (!RequestZenShutdownOnEffectivePort(EffectiveListenPort))
 	{
-		if (FPlatformProcess::IsProcRunning(ProcessHandle) || IsZenProcessUsingPort(ProcessPort))
+		if (FPlatformProcess::IsProcRunning(ProcessHandle) || IsZenProcessUsingEffectivePort(EffectiveListenPort))
 		{
-			UE_LOG(LogZenServiceInstance, Warning, TEXT("Failed to request shutdown for running service using port %u (Pid: %u)"), ProcessPort, ServicePid);
+			UE_LOG(LogZenServiceInstance, Warning, TEXT("Failed to request shutdown for running service using port %u (Pid: %u)"), EffectiveListenPort, ServicePid);
 			return false;
 		}
 		return true;
 	}
 	uint64 ZenShutdownWaitStartTime = FPlatformTime::Cycles64();
-	while (FPlatformProcess::IsProcRunning(ProcessHandle) || IsZenProcessUsingPort(ProcessPort))
+	while (FPlatformProcess::IsProcRunning(ProcessHandle) || IsZenProcessUsingEffectivePort(EffectiveListenPort))
 	{
 		double ZenShutdownWaitDuration = FPlatformTime::ToSeconds64(FPlatformTime::Cycles64() - ZenShutdownWaitStartTime);
 		if (ZenShutdownWaitDuration < MaximumWaitDurationSeconds)
@@ -1009,7 +1086,7 @@ ShutdownRunningService(uint16 ProcessPort, double MaximumWaitDurationSeconds = 5
 		}
 		else
 		{
-			UE_LOG(LogZenServiceInstance, Warning, TEXT("Timed out waiting for shutdown of running service using port %u (Pid: %u)"), ProcessPort, ServicePid);
+			UE_LOG(LogZenServiceInstance, Warning, TEXT("Timed out waiting for shutdown of running service using port %u (Pid: %u)"), EffectiveListenPort, ServicePid);
 			return false;
 		}
 	}
@@ -1147,10 +1224,19 @@ FZenLocalServiceRunContext::WriteToJsonFile(const TCHAR* Filename) const
 }
 
 bool
-IsLocalServiceRunning(const TCHAR* DataPath, uint16* OutPort)
+IsLocalServiceRunning(const TCHAR* DataPath, uint16* OutEffectivePort)
 {
 	const FString LockFilePath = FPaths::Combine(DataPath, TEXT(".lock"));
-	return IsZenProcessUsingDataDir(*LockFilePath, OutPort);
+	LockFileData LockFileState;
+	if (IsZenProcessUsingDataDir(*LockFilePath, &LockFileState))
+	{
+		if (LockFileState.IsValid && LockFileState.IsReady)
+		{
+			*OutEffectivePort = LockFileState.EffectivePort;
+		}
+		return true;
+	}
+	return false;
 }
 
 FProcHandle
@@ -1243,10 +1329,15 @@ StartLocalService(const FZenLocalServiceRunContext& Context, const TCHAR* Transi
 bool
 StopLocalService(const TCHAR* DataPath, double MaximumWaitDurationSeconds)
 {
-	uint16 CurrentPort = 0;
-	if (IsLocalServiceRunning(DataPath, &CurrentPort))
+	const FString LockFilePath = FPaths::Combine(DataPath, TEXT(".lock"));
+	LockFileData LockFileState;
+	if (IsZenProcessUsingDataDir(*LockFilePath, &LockFileState))
 	{
-		return ShutdownRunningService(CurrentPort, MaximumWaitDurationSeconds);
+		if (LockFileState.IsValid)
+		{
+			return ShutdownRunningServiceUsingEffectivePort(LockFileState.EffectivePort, MaximumWaitDurationSeconds);
+		}
+		return false;
 	}
 	return true;
 }
@@ -1423,7 +1514,7 @@ FZenServiceInstance::TryRecovery()
 			FZenLocalServiceRunContext RunContext;
 			if (TryGetLocalServiceRunContext(RunContext))
 			{
-				if (!ShutdownRunningService(*RunContext.GetDataPath()))
+				if (!ShutdownRunningServiceUsingExecutablePath(*RunContext.GetExecutable()))
 				{
 					return false;
 				}
@@ -1524,7 +1615,7 @@ PromptUserOfLockedDataFolder(const FString& DataPath)
 
 	FText ZenLaunchFailurePromptTitle = NSLOCTEXT("Zen", "Zen_NonLocalProcessUsesDataDirPromptTitle", "Failed to launch");
 	FText ZenLaunchFailurePromptText = FText::Format(NSLOCTEXT("Zen", "Zen_NonLocalProcessUsesDataDirPromptText", "ZenServer Failed to auto launch, an unknown process is locking the data folder '{0}'"), FText::FromString(DataPath));
-	FPlatformMisc::MessageBoxExt(EAppMsgType::Ok, *ZenLaunchFailurePromptTitle.ToString(), *ZenLaunchFailurePromptText.ToString());
+	FPlatformMisc::MessageBoxExt(EAppMsgType::Ok, *ZenLaunchFailurePromptText.ToString(), *ZenLaunchFailurePromptTitle.ToString());
 }
 
 static void
@@ -1540,7 +1631,7 @@ PromptUserOfFailedShutDownOfExistingProcess(uint16 Port)
 
 	FText ZenLaunchFailurePromptTitle = NSLOCTEXT("Zen", "Zen_ShutdownFailurePromptTitle", "Failed to launch");
 	FText ZenLaunchFailurePromptText = FText::Format(NSLOCTEXT("Zen", "Zen_ShutdownFailurePromptText", "ZenServer Failed to auto launch, failed to shut down currently running service using port '{0}'"), Port);
-	FPlatformMisc::MessageBoxExt(EAppMsgType::Ok, *ZenLaunchFailurePromptTitle.ToString(), *ZenLaunchFailurePromptText.ToString());
+	FPlatformMisc::MessageBoxExt(EAppMsgType::Ok, *ZenLaunchFailurePromptText.ToString(), *ZenLaunchFailurePromptTitle.ToString());
 }
 
 
@@ -1560,7 +1651,7 @@ FZenServiceInstance::ConditionalUpdateLocalInstall()
 	if (IsInstallVersionOutOfDate(InTreeUtilityPath, InstallUtilityPath, InTreeServicePath, InstallServicePath, InTreeVersionCache, InstallVersionCache))
 	{
 		UE_LOG(LogZenServiceInstance, Display, TEXT("Installing service from '%s' to '%s'"), *InTreeServicePath, *InstallServicePath);
-		if (!ShutdownRunningService(*InstallServicePath))
+		if (!ShutdownRunningServiceUsingExecutablePath(*InstallServicePath))
 		{
 			PromptUserToStopRunningServerInstanceForUpdate(InstallServicePath);
 			return FString();
@@ -1608,16 +1699,15 @@ FZenServiceInstance::ConditionalUpdateLocalInstall()
 bool
 FZenServiceInstance::AutoLaunch(const FServiceAutoLaunchSettings& InSettings, FString&& ExecutablePath, FString& OutHostName, uint16& OutPort)
 {
-	uint16 DesiredPort = InSettings.DesiredPort;
 	IFileManager& FileManager = IFileManager::Get();
 	const FString LockFilePath = FPaths::Combine(InSettings.DataPath, TEXT(".lock"));
 	const FString ExecutionContextFilePath = FPaths::SetExtension(ExecutablePath, TEXT(".runcontext"));
 
 	FString WorkingDirectory = FPaths::GetPath(ExecutablePath);
 
-	uint16 RunningProcessPort = 0;
+	LockFileData LockFileState;
 	uint64 ZenWaitForRunningProcessReadyStartTime = FPlatformTime::Cycles64();
-	while (IsZenProcessUsingDataDir(*LockFilePath, &RunningProcessPort) && RunningProcessPort == 0)
+	while (IsZenProcessUsingDataDir(*LockFilePath, &LockFileState) && LockFileState.IsValid && !LockFileState.IsReady)
 	{
 		// Server is starting up, wait for it to get ready
 		double ZenWaitDuration = FPlatformTime::ToSeconds64(FPlatformTime::Cycles64() - ZenWaitForRunningProcessReadyStartTime);
@@ -1628,12 +1718,13 @@ FZenServiceInstance::AutoLaunch(const FServiceAutoLaunchSettings& InSettings, FS
 		FPlatformProcess::Sleep(0.1f);
 	}
 
-	uint16 ShutdownPort = 0;
+	uint16 ShutdownEffectivePort = 0;
 	bool bLaunchNewInstance = true;
 
-	if (RunningProcessPort != 0)
+	if (LockFileState.IsReady)
 	{
-		if (!IsZenProcessUsingPort(RunningProcessPort))
+		const ZenServerState State;
+		if (State.LookupByPid(LockFileState.ProcessId) == nullptr)
 		{
 			PromptUserOfLockedDataFolder(*InSettings.DataPath);
 			FPlatformMisc::RequestExit(true);
@@ -1642,7 +1733,7 @@ FZenServiceInstance::AutoLaunch(const FServiceAutoLaunchSettings& InSettings, FS
 
 		FZenLocalServiceRunContext DesiredRunContext;
 		DesiredRunContext.Executable = ExecutablePath;
-		DesiredRunContext.CommandlineArguments = DetermineCmdLineWithoutTransientComponents(InSettings, RunningProcessPort);
+		DesiredRunContext.CommandlineArguments = DetermineCmdLineWithoutTransientComponents(InSettings, InSettings.DesiredPort);
 		DesiredRunContext.WorkingDirectory = WorkingDirectory;
 		DesiredRunContext.DataPath = InSettings.DataPath;
 		DesiredRunContext.bShowConsole = InSettings.bShowConsole;
@@ -1650,31 +1741,35 @@ FZenServiceInstance::AutoLaunch(const FServiceAutoLaunchSettings& InSettings, FS
 		FZenLocalServiceRunContext CurrentRunContext;
 		if (CurrentRunContext.ReadFromJsonFile(*ExecutionContextFilePath) && (DesiredRunContext == CurrentRunContext))
 		{
-			UE_LOG(LogZenServiceInstance, Log, TEXT("Found existing instance running on port %u matching our settings"), RunningProcessPort);
-			DesiredPort = RunningProcessPort;
+			UE_LOG(LogZenServiceInstance, Log, TEXT("Found existing instance running on port %u matching our settings"), InSettings.DesiredPort);
 			bLaunchNewInstance = false;
 		}
 		else
 		{
-			UE_LOG(LogZenServiceInstance, Log, TEXT("Found existing instance running on port %u with different run context, will attempt shut down"), RunningProcessPort);
-			ShutdownPort = RunningProcessPort;
+			UE_LOG(LogZenServiceInstance, Log, TEXT("Found existing instance running on port %u with different run context, will attempt shut down"), InSettings.DesiredPort);
+			ShutdownEffectivePort = LockFileState.EffectivePort;
 		}
 	}
-	else if (IsZenProcessUsingPort(DesiredPort))
+	else
 	{
-		UE_LOG(LogZenServiceInstance, Log, TEXT("Found existing instance running on port %u with different data directory, will attempt shut down"), DesiredPort);
-		ShutdownPort = DesiredPort;
+		const ZenServerState State;
+		const ZenServerState::ZenServerEntry* RunningEntry = State.LookupByDesiredListenPort(InSettings.DesiredPort);
+		if (RunningEntry != nullptr)
+		{
+			ShutdownEffectivePort = RunningEntry->EffectiveListenPort;
+			UE_LOG(LogZenServiceInstance, Log, TEXT("Found existing instance running on port %u with different data directory, will attempt shut down"), ShutdownEffectivePort);
+		}
 	}
 
-	if (ShutdownPort != 0)
+	if (ShutdownEffectivePort != 0)
 	{
-		if (!ShutdownRunningService(ShutdownPort))
+		if (!ShutdownRunningServiceUsingEffectivePort(ShutdownEffectivePort))
 		{
-			PromptUserOfFailedShutDownOfExistingProcess(ShutdownPort);
+			PromptUserOfFailedShutDownOfExistingProcess(ShutdownEffectivePort);
 			FPlatformMisc::RequestExit(true);
 			return false;
 		}
-		checkf(!IsZenProcessUsingPort(ShutdownPort), TEXT("Service port is still in use %u"), ShutdownPort);
+		checkf(!IsZenProcessUsingEffectivePort(ShutdownEffectivePort), TEXT("Service port is still in use %u"), ShutdownEffectivePort);
 
 		if (IsLockFileLocked(*LockFilePath))
 		{
@@ -1686,9 +1781,9 @@ FZenServiceInstance::AutoLaunch(const FServiceAutoLaunchSettings& InSettings, FS
 
 	// When limiting process lifetime, always re-launch to add sponsor process IDs.
 	// When not limiting process lifetime, only launch if the process is not already live.
-	if (bLaunchNewInstance || InSettings.bLimitProcessLifetime)
+	if (bLaunchNewInstance)
 	{
-		FString ParmsWithoutTransients = DetermineCmdLineWithoutTransientComponents(InSettings, DesiredPort);
+		FString ParmsWithoutTransients = DetermineCmdLineWithoutTransientComponents(InSettings, InSettings.DesiredPort);
 		FString TransientParms;
 
 		if (InSettings.bLimitProcessLifetime)
@@ -1713,14 +1808,27 @@ FZenServiceInstance::AutoLaunch(const FServiceAutoLaunchSettings& InSettings, FS
 
 		if (!Proc.IsValid())
 		{
-			UE_LOG(LogZenServiceInstance, Warning, TEXT("Failed launch service using executable '%s' on port %u"), *ExecutablePath, DesiredPort);
+			UE_LOG(LogZenServiceInstance, Warning, TEXT("Failed launch service using executable '%s' on port %u"), *ExecutablePath, InSettings.DesiredPort);
 			return false;
 		}
 	}
+	else if (InSettings.bLimitProcessLifetime)
+	{
+		ZenServerState State(/*ReadOnly*/ false);
+		ZenServerState::ZenServerEntry* RunningEntry = State.LookupByDesiredListenPort(InSettings.DesiredPort);
+		if (RunningEntry == nullptr)
+		{
+			UE_LOG(LogZenServiceInstance, Warning, TEXT("Failed attach as sponsor process to executable '%s' on port %u, can't find entry in shared state"), *ExecutablePath, InSettings.DesiredPort);
+		}
+		else if (!RunningEntry->AddSponsorProcess(FPlatformProcess::GetCurrentProcessId()))
+		{
+			UE_LOG(LogZenServiceInstance, Warning, TEXT("Failed attach as sponsor process to executable '%s' on port %u, to many sponsored processes attached already"), *ExecutablePath, InSettings.DesiredPort);
+		}
+	}
 
 	OutHostName = TEXT("[::1]");
 	// Default to assuming that we get to run on the port we want
-	OutPort = DesiredPort;
+	OutPort = InSettings.DesiredPort;
 
 	FScopedSlowTask WaitForZenReadySlowTask(0, NSLOCTEXT("Zen", "Zen_WaitingForReady", "Waiting for ZenServer to be ready"));
 	uint64 ZenWaitStartTime = FPlatformTime::Cycles64();
@@ -1733,13 +1841,13 @@ FZenServiceInstance::AutoLaunch(const FServiceAutoLaunchSettings& InSettings, FS
 	bool bIsReady = false;
 	while (!bIsReady)
 	{
-		FCbObject LockObject;
-		if (ReadCbLockFile(LockFilePath, LockObject))
+		LockFileData RunningLockFileState = ReadCbLockFile(LockFilePath);
+		if (RunningLockFileState.IsValid)
 		{
-			bIsReady = LockObject["ready"].AsBool();
+			bIsReady = RunningLockFileState.IsReady;
 			if (bIsReady)
 			{
-				OutPort = LockObject["port"].AsUInt16(DesiredPort);
+				OutPort = RunningLockFileState.EffectivePort;
 				break;
 			}
 		}
