Accept duplicated file handles in Zen RPC responses

diff --git a/Engine/Source/Developer/DerivedDataCache/Private/ZenCacheStore.cpp b/Engine/Source/Developer/DerivedDataCache/Private/ZenCacheStore.cpp
--- a/Engine/Source/Developer/DerivedDataCache/Private/ZenCacheStore.cpp
+++ b/Engine/Source/Developer/DerivedDataCache/Private/ZenCacheStore.cpp
@@ -366,6 +366,7 @@ public:
 					if (CacheStore.bIsLocalConnection)
 					{
 						BatchRequest.AddInteger(ANSITEXTVIEW("AcceptFlags"), static_cast<uint32_t>(Zen::Http::RpcAcceptOptions::kAllowLocalReferences));
+						BatchRequest.AddInteger(ANSITEXTVIEW("Pid"), FPlatformProcess::GetCurrentProcessId());
 					}
 
 					BatchRequest.BeginObject(ANSITEXTVIEW("Params"));
@@ -660,6 +661,7 @@ public:
 				if (CacheStore.bIsLocalConnection)
 				{
 					BatchRequest.AddInteger(ANSITEXTVIEW("AcceptFlags"), static_cast<uint32_t>(Zen::Http::RpcAcceptOptions::kAllowLocalReferences));
+					BatchRequest.AddInteger(ANSITEXTVIEW("Pid"), FPlatformProcess::GetCurrentProcessId());
 				}
 
 				BatchRequest.BeginObject(ANSITEXTVIEW("Params"));
@@ -819,6 +821,7 @@ public:
 				if (CacheStore.bIsLocalConnection)
 				{
 					BatchRequest.AddInteger(ANSITEXTVIEW("AcceptFlags"), static_cast<uint32_t>(Zen::Http::RpcAcceptOptions::kAllowLocalReferences));
+					BatchRequest.AddInteger(ANSITEXTVIEW("Pid"), FPlatformProcess::GetCurrentProcessId());
 				}
 
 				BatchRequest.BeginObject(ANSITEXTVIEW("Params"));
diff --git a/Engine/Source/Developer/Zen/Private/ZenSerialization.cpp b/Engine/Source/Developer/Zen/Private/ZenSerialization.cpp
--- a/Engine/Source/Developer/Zen/Private/ZenSerialization.cpp
+++ b/Engine/Source/Developer/Zen/Private/ZenSerialization.cpp
@@ -3,12 +3,17 @@
 #include "ZenSerialization.h"
 #include "Compression/OodleDataCompression.h"
 #include "HAL/FileManager.h"
+#include "Misc/ScopeExit.h"
 #include "Serialization/CompactBinary.h"
 #include "Serialization/CompactBinaryPackage.h"
 #include "Serialization/CompactBinarySerialization.h"
 #include "Serialization/CompactBinaryValidation.h"
 #include "Serialization/CompactBinaryWriter.h"
 
+#if PLATFORM_WINDOWS
+#	include "Windows/WindowsHWrapper.h"
+#endif // PLATFORM_WINDOWS
+
 #if UE_WITH_ZEN
 
 DEFINE_LOG_CATEGORY_STATIC(LogZenSerialization, Log, All);
@@ -313,16 +318,74 @@ bool TryLoadCbPackage(FCbPackage& Package, FArchive& Ar, FCbBufferAllocator Allo
 				return false;
 			}
 
+			FCompressedBuffer CompBuf;
 			FString Path(FUtf8StringView(PathPointer, static_cast<uint32_t>(AttachRefHdr->AbsolutePathLength)));
-			TUniquePtr<FArchive> ChunkReader(IFileManager::Get().CreateFileReader(*Path, FILEREAD_Silent));
-			if (!ChunkReader)
+
+			const FStringView HandlePrefix(TEXT(":?#:"));
+			if (Path.StartsWith(HandlePrefix))
 			{
-				// File does not exist on disk - treat it as missing
-				UE_LOG(LogZenSerialization, Warning, TEXT("Unable to read local file '%s', treating it as a missing attachment"), *Path);
-				++Index;
-				continue;
+#if PLATFORM_WINDOWS
+				FString		HandleString(Path.RightChop(HandlePrefix.Len()));
+
+				uint64 HandleNum = 0;
+				LexFromString(HandleNum, *HandleString);
+				if (HandleNum != 0)
+				{
+					HANDLE Handle = reinterpret_cast<HANDLE>(HandleNum);
+					ON_SCOPE_EXIT{ CloseHandle(Handle); };
+					LARGE_INTEGER FileSize;
+					BOOL OK = GetFileSizeEx(Handle, &FileSize);
+					if (!OK)
+					{
+						UE_LOG(LogZenSerialization, Warning, TEXT("Unable to read local file via handle '%u', treating it as a missing attachment"), HandleNum);
+						++Index;
+						continue;
+					}
+					FUniqueBuffer MemBuffer = Allocator(FileSize.QuadPart);
+
+					LONGLONG Offset = 0;
+					while (Offset < FileSize.QuadPart)
+					{
+						DWORD Length = static_cast<DWORD>(FMath::Min<LONGLONG>(0xffff'ffffu, FileSize.QuadPart - Offset));
+						DWORD dwNumberOfBytesRead = 0;
+						OVERLAPPED Ovl{};
+						Ovl.Offset = DWORD(Offset & 0xffff'ffffu);
+						Ovl.OffsetHigh = DWORD(Offset >> 32);
+						BOOL  Success = ::ReadFile(Handle, (void*)MemBuffer.GetView().Mid(Offset, Length).GetData(), Length, &dwNumberOfBytesRead, &Ovl);
+						if (!Success)
+						{
+							DWORD LastError = GetLastError();
+							UE_LOG(LogZenSerialization, Warning, TEXT("Unable to read local file via handle '%u', error: %u, treating it as a missing attachment"), HandleNum, LastError);
+							break;
+						}
+						Offset += Length;
+					}
+					if (Offset != FileSize.QuadPart)
+					{
+						++Index;
+						continue;
+					}
+					CompBuf = FCompressedBuffer::FromCompressed(MemBuffer.MoveToShared());
+				}
+#else // PLATFORM_WINDOWS
+				// Only suported on Windows. For Linux we could potentially use pidfd_getfd() but that requires a fairly new Linux kernel/includes and
+				// to deal with acceess rights etc.
+				checkf(false, TEXT("Passing file handles is not supported on this platform"));
+#endif // PLATFORM_WINDOWS
 			}
-			FCompressedBuffer CompBuf = FCompressedBuffer::Load(*ChunkReader);
+			else
+			{
+				TUniquePtr<FArchive> ChunkReader(IFileManager::Get().CreateFileReader(*Path, FILEREAD_Silent));
+				if (!ChunkReader)
+				{
+					// File does not exist on disk - treat it as missing
+					UE_LOG(LogZenSerialization, Warning, TEXT("Unable to read local file '%s', treating it as a missing attachment"), *Path);
+					++Index;
+					continue;
+				}
+				CompBuf = FCompressedBuffer::Load(*ChunkReader);
+			}
+
 			if (!CompBuf)
 			{
 				// File has wrong format - treat it as an error
