Fix up of Event Attribute names to support Grafana in StudioAnalytics, DerivedData, Zen, Virtualization and CookProfiling. Grafana does not like the . separator so these have been replaced by _. I have checked this with the DPE team ( and  verified ) this is still compatible with Snowflake/Tableau telemetry and such will not requires any changes to Snowflake or exisitng Tableau dashboards.

diff --git a/Engine/Source/Developer/DerivedDataCache/Private/DerivedDataCache.cpp b/Engine/Source/Developer/DerivedDataCache/Private/DerivedDataCache.cpp
--- a/Engine/Source/Developer/DerivedDataCache/Private/DerivedDataCache.cpp
+++ b/Engine/Source/Developer/DerivedDataCache/Private/DerivedDataCache.cpp
@@ -964,37 +964,37 @@ public:
 		// Append to the attributes
 		for (const FDerivedDataCacheResourceStat& Stat : ResourceStats)
 		{
-			FString BaseName = TEXT("DDC.Resource.") + Stat.AssetType;
+			FString BaseName = TEXT("DDC_Resource_") + Stat.AssetType;
 
 			BaseName = BaseName.Replace(TEXT("("), TEXT("")).Replace(TEXT(")"), TEXT(""));
 
 			{
-				FString AttrName = BaseName + TEXT(".BuildCount");
+				FString AttrName = BaseName + TEXT("_BuildCount");
 				Attributes.Emplace(MoveTemp(AttrName), Stat.BuildCount);
 			}
 
 			{
-				FString AttrName = BaseName + TEXT(".BuildTimeSec");
+				FString AttrName = BaseName + TEXT("_BuildTimeSec");
 				Attributes.Emplace(MoveTemp(AttrName), Stat.BuildTimeSec);
 			}
 
 			{
-				FString AttrName = BaseName + TEXT(".BuildSizeMB");
+				FString AttrName = BaseName + TEXT("_BuildSizeMB");
 				Attributes.Emplace(MoveTemp(AttrName), Stat.BuildSizeMB);
 			}
 
 			{
-				FString AttrName = BaseName + TEXT(".LoadCount");
+				FString AttrName = BaseName + TEXT("_LoadCount");
 				Attributes.Emplace(MoveTemp(AttrName), Stat.LoadCount);
 			}
 
 			{
-				FString AttrName = BaseName + TEXT(".LoadTimeSecLoadTimeSec");
+				FString AttrName = BaseName + TEXT("_LoadTimeSecLoadTimeSec");
 				Attributes.Emplace(MoveTemp(AttrName), Stat.LoadTimeSec);
 			}
 
 			{
-				FString AttrName = BaseName + TEXT(".LoadSizeMB");
+				FString AttrName = BaseName + TEXT("_LoadSizeMB");
 				Attributes.Emplace(MoveTemp(AttrName), Stat.LoadSizeMB);
 			}
 		}
@@ -1007,7 +1007,7 @@ public:
 		// Append to the attributes
 		for (const FDerivedDataCacheSummaryStat& Stat : SummaryStats.Stats)
 		{
-			FString FormattedAttrName = "DDC.Summary." + Stat.Key;
+			FString FormattedAttrName = "DDC_Summary_" + Stat.Key;
 
 			if (Stat.Value.IsNumeric())
 			{
diff --git a/Engine/Source/Developer/Virtualization/Private/VirtualizationManager.cpp b/Engine/Source/Developer/Virtualization/Private/VirtualizationManager.cpp
--- a/Engine/Source/Developer/Virtualization/Private/VirtualizationManager.cpp
+++ b/Engine/Source/Developer/Virtualization/Private/VirtualizationManager.cpp
@@ -2105,52 +2105,52 @@ void FVirtualizationManager::GatherAnalytics(TArray<FAnalyticsEventAttribute>& A
 		const FString BaseName = TEXT("Virtualization");
 
 		{
-			FString AttrName = BaseName + TEXT(".Enabled");
+			FString AttrName = BaseName + TEXT("_Enabled");
 			Attributes.Emplace(MoveTemp(AttrName), System.IsEnabled());
 		}
 
 		{
-			FString AttrName = BaseName + TEXT(".Cache.TimeSpent");
+			FString AttrName = BaseName + TEXT("_Cache_TimeSpent");
 			Attributes.Emplace(MoveTemp(AttrName), (double)PayloadActivityInfo.Cache.CyclesSpent * FPlatformTime::GetSecondsPerCycle());
 		}
 
 		{
-			FString AttrName = BaseName + TEXT(".Cache.PayloadCount");
+			FString AttrName = BaseName + TEXT("_Cache_PayloadCount");
 			Attributes.Emplace(MoveTemp(AttrName), (double)PayloadActivityInfo.Cache.PayloadCount);
 		}
 
 		{
-			FString AttrName = BaseName + TEXT(".Cache.TotalBytes");
+			FString AttrName = BaseName + TEXT("_Cache_TotalBytes");
 			Attributes.Emplace(MoveTemp(AttrName), (double)PayloadActivityInfo.Cache.TotalBytes);
 		}
 
 		{
-			FString AttrName = BaseName + TEXT(".Push.TimeSpent");
+			FString AttrName = BaseName + TEXT("_Push_TimeSpent");
 			Attributes.Emplace(MoveTemp(AttrName), (double)PayloadActivityInfo.Push.CyclesSpent * FPlatformTime::GetSecondsPerCycle());
 		}
 
 		{
-			FString AttrName = BaseName + TEXT(".Push.PayloadCount");
+			FString AttrName = BaseName + TEXT("_Push_PayloadCount");
 			Attributes.Emplace(MoveTemp(AttrName), (double)PayloadActivityInfo.Push.PayloadCount);
 		}
 
 		{
-			FString AttrName = BaseName + TEXT(".Push.TotalBytes");
+			FString AttrName = BaseName + TEXT("_Push_TotalBytes");
 			Attributes.Emplace(MoveTemp(AttrName), (double)PayloadActivityInfo.Push.TotalBytes);
 		}
 
 		{
-			FString AttrName = BaseName + TEXT(".Pull.TimeSpent");
+			FString AttrName = BaseName + TEXT("_Pull_TimeSpent");
 			Attributes.Emplace(MoveTemp(AttrName), (double)PayloadActivityInfo.Pull.CyclesSpent * FPlatformTime::GetSecondsPerCycle());
 		}
 
 		{
-			FString AttrName = BaseName + TEXT(".Pull.PayloadCount");
+			FString AttrName = BaseName + TEXT("_Pull_PayloadCount");
 			Attributes.Emplace(MoveTemp(AttrName), (double)PayloadActivityInfo.Pull.PayloadCount);
 		}
 
 		{
-			FString AttrName = BaseName + TEXT(".Pull.TotalBytes");
+			FString AttrName = BaseName + TEXT("_Pull_TotalBytes");
 			Attributes.Emplace(MoveTemp(AttrName), (double)PayloadActivityInfo.Pull.TotalBytes);
 		}
 	}
diff --git a/Engine/Source/Developer/Zen/Private/ZenServerInterface.cpp b/Engine/Source/Developer/Zen/Private/ZenServerInterface.cpp
--- a/Engine/Source/Developer/Zen/Private/ZenServerInterface.cpp
+++ b/Engine/Source/Developer/Zen/Private/ZenServerInterface.cpp
@@ -2151,80 +2151,80 @@ FZenServiceInstance::GatherAnalytics(TArray<FAnalyticsEventAttribute>& Attribute
 	if (GetStats(ZenStats) == false)
 		return false;
 
-	const FString BaseName = TEXT("Zen");
+	const FString BaseName = TEXT("Zen_");
 
 	{
-		FString AttrName = BaseName + TEXT(".Enabled");
+		FString AttrName = BaseName + TEXT("Enabled");
 		Attributes.Emplace(MoveTemp(AttrName), ZenStats.IsValid);
 	}
 
 	{
-		FString AttrName = BaseName + TEXT(".Cache.HitRatio");
+		FString AttrName = BaseName + TEXT("Cache_HitRatio");
 		Attributes.Emplace(MoveTemp(AttrName), ZenStats.CacheStats.HitRatio);
 	}
 
 	{
-		FString AttrName = BaseName + TEXT(".Cache.Hits");
+		FString AttrName = BaseName + TEXT("Cache_Hits");
 		Attributes.Emplace(MoveTemp(AttrName), ZenStats.CacheStats.Hits);
 	}
 
 	{
-		FString AttrName = BaseName + TEXT(".Cache.Misses");
+		FString AttrName = BaseName + TEXT("Cache_Misses");
 		Attributes.Emplace(MoveTemp(AttrName), ZenStats.CacheStats.Misses);
 	}
 
 	{
-		FString AttrName = BaseName + TEXT(".Cache.Size.Disk");
+		FString AttrName = BaseName + TEXT("Cache_Size_Disk");
 		Attributes.Emplace(MoveTemp(AttrName), ZenStats.CacheStats.Size.Disk);
 	}
 
 	{
-		FString AttrName = BaseName + TEXT(".Cache.Size.Memory");
+		FString AttrName = BaseName + TEXT("Cache_Size_Memory");
 		Attributes.Emplace(MoveTemp(AttrName), ZenStats.CacheStats.Size.Memory);
 	}
 
 	{
-		FString AttrName = BaseName + TEXT(".Cache.UpstreamHits");
+		FString AttrName = BaseName + TEXT("Cache_UpstreamHits");
 		Attributes.Emplace(MoveTemp(AttrName), ZenStats.CacheStats.UpstreamHits);
 	}
 
 	{
-		FString AttrName = BaseName + TEXT(".Cache.UpstreamRatio");
+		FString AttrName = BaseName + TEXT("Cache_UpstreamRatio");
 		Attributes.Emplace(MoveTemp(AttrName), ZenStats.CacheStats.UpstreamRatio);
 	}
 
 	{
-		FString AttrName = BaseName + TEXT(".Cache.TotalUploadedMB");
+		FString AttrName = BaseName + TEXT("Cache_TotalUploadedMB");
 		Attributes.Emplace(MoveTemp(AttrName), ZenStats.UpstreamStats.TotalUploadedMB);
 	}
 
 	{
-		FString AttrName = BaseName + TEXT(".Upstream.TotalDownloadedMB");
+		FString AttrName = BaseName + TEXT("Upstream_TotalDownloadedMB");
 		Attributes.Emplace(MoveTemp(AttrName), ZenStats.UpstreamStats.TotalDownloadedMB);
 	}
 
 	{
-		FString AttrName = BaseName + TEXT(".Upstream.TotalUploadedMB");
+		FString AttrName = BaseName + TEXT("Upstream_TotalUploadedMB");
 		Attributes.Emplace(MoveTemp(AttrName), ZenStats.UpstreamStats.TotalUploadedMB);
 	}
 
 	{
-		FString AttrName = BaseName + TEXT(".Cas.Size.Large");
+		FString AttrName = BaseName + TEXT("Cas_Size_Large");
 		Attributes.Emplace(MoveTemp(AttrName), ZenStats.CASStats.Size.Large);
 	}
 
 	{
-		FString AttrName = BaseName + TEXT(".Cas.Size.Small");
+		FString AttrName = BaseName + TEXT("Cas_Size_Small");
 		Attributes.Emplace(MoveTemp(AttrName), ZenStats.CASStats.Size.Small);
 	}
 
 	{
-		FString AttrName = BaseName + TEXT(".Cas.Size.Tiny");
+		FString AttrName = BaseName + TEXT("Cas_Size_Tiny");
 		Attributes.Emplace(MoveTemp(AttrName), ZenStats.CASStats.Size.Tiny);
 	}
 
 	{
-		FString AttrName = BaseName + TEXT(".Cas.Size.Total");
+		FString AttrName = BaseName + TEXT("Cas_Size_Total");
 		Attributes.Emplace(MoveTemp(AttrName), ZenStats.CASStats.Size.Total);
 	}
 
diff --git a/Engine/Source/Editor/UnrealEd/Private/Cooker/CookProfiling.cpp b/Engine/Source/Editor/UnrealEd/Private/Cooker/CookProfiling.cpp
--- a/Engine/Source/Editor/UnrealEd/Private/Cooker/CookProfiling.cpp
+++ b/Engine/Source/Editor/UnrealEd/Private/Cooker/CookProfiling.cpp
@@ -920,7 +920,7 @@ void LogCookStats(ECookMode::Type CookMode)
 		{
 			for (const auto& Attr : StatAttributes)
 			{
-				FString FormattedAttrName = StatName + "." + Attr.Key;
+				FString FormattedAttrName = StatName + "_" + Attr.Key;
 
 				if (Attr.Value.IsNumeric())
 				{
