Add --no-sentry option to zenserver launch when "Send Unattended Bug Reports" is disabled in Unreal Editor.

The setting will only take affect for zenserver when the editor is restarted.

diff --git a/Engine/Source/Developer/Zen/Private/ZenServerInterface.cpp b/Engine/Source/Developer/Zen/Private/ZenServerInterface.cpp
--- a/Engine/Source/Developer/Zen/Private/ZenServerInterface.cpp
+++ b/Engine/Source/Developer/Zen/Private/ZenServerInterface.cpp
@@ -579,6 +579,7 @@ FServiceSettings::ReadFromConfig()
 			ReadUInt16FromConfig(AutoLaunchConfigSection, TEXT("DesiredPort"), AutoLaunchSettings.DesiredPort, GEngineIni);
 			GConfig->GetBool(AutoLaunchConfigSection, TEXT("ShowConsole"), AutoLaunchSettings.bShowConsole, GEngineIni);
 			GConfig->GetBool(AutoLaunchConfigSection, TEXT("LimitProcessLifetime"), AutoLaunchSettings.bLimitProcessLifetime, GEngineIni);
+			GConfig->GetBool(TEXT("/Script/UnrealEd.CrashReportsPrivacySettings"), TEXT("bSendUnattendedBugReports"), AutoLaunchSettings.bSendUnattendedBugReports, GEditorSettingsIni);
 		}
 	}
 	else
@@ -610,6 +611,7 @@ FServiceSettings::ReadFromCompactBinary(FCbFieldView Field)
 				AutoLaunchSettings.DesiredPort = AutoLaunchSettingsObject["DesiredPort"].AsInt16();
 				AutoLaunchSettings.bShowConsole = AutoLaunchSettingsObject["ShowConsole"].AsBool();
 				AutoLaunchSettings.bLimitProcessLifetime = AutoLaunchSettingsObject["LimitProcessLifetime"].AsBool();
+				AutoLaunchSettings.bSendUnattendedBugReports = AutoLaunchSettingsObject["SendUnattendedBugReports"].AsBool();
 			}
 		}
 	}
@@ -665,6 +667,7 @@ FServiceSettings::WriteToCompactBinary(FCbWriter& Writer) const
 		Writer << "DesiredPort" << AutoLaunchSettings.DesiredPort;
 		Writer << "ShowConsole" << AutoLaunchSettings.bShowConsole;
 		Writer << "LimitProcessLifetime" << AutoLaunchSettings.bLimitProcessLifetime;
+		Writer << "SendUnattendedBugReports" << AutoLaunchSettings.bSendUnattendedBugReports;
 		Writer.EndObject();
 	}
 	else
@@ -1123,6 +1126,11 @@ DetermineCmdLineWithoutTransientComponents(const FServiceAutoLaunchSettings& InS
 		}
 	}
 
+	if (!InSettings.bSendUnattendedBugReports)
+	{
+		Parms.Append(TEXT(" --no-sentry"));
+	}
+
 	return Parms;
 }
 
diff --git a/Engine/Source/Developer/Zen/Public/Experimental/ZenServerInterface.h b/Engine/Source/Developer/Zen/Public/Experimental/ZenServerInterface.h
--- a/Engine/Source/Developer/Zen/Public/Experimental/ZenServerInterface.h
+++ b/Engine/Source/Developer/Zen/Public/Experimental/ZenServerInterface.h
@@ -42,6 +42,7 @@ struct FServiceAutoLaunchSettings
 	uint16 DesiredPort = 1337;
 	bool bShowConsole = false;
 	bool bLimitProcessLifetime = false;
+	bool bSendUnattendedBugReports = false;
 };
 
 struct FServiceSettings
