Fix potential crash when checking if a local Zen service is running if the 'OutEffectivePort' parameter is not set.

#rnx

- The parameter OutEffectivePort is optional, but we were not checking if it was valid or not before writing to it.

diff --git a/Engine/Source/Developer/Zen/Private/ZenServerInterface.cpp b/Engine/Source/Developer/Zen/Private/ZenServerInterface.cpp
--- a/Engine/Source/Developer/Zen/Private/ZenServerInterface.cpp
+++ b/Engine/Source/Developer/Zen/Private/ZenServerInterface.cpp
@@ -1231,7 +1231,7 @@ IsLocalServiceRunning(const TCHAR* DataPath, uint16* OutEffectivePort)
 	LockFileData LockFileState;
 	if (IsZenProcessUsingDataDir(*LockFilePath, &LockFileState))
 	{
-		if (LockFileState.IsValid && LockFileState.IsReady)
+		if (OutEffectivePort != nullptr && LockFileState.IsValid && LockFileState.IsReady)
 		{
 			*OutEffectivePort = LockFileState.EffectivePort;
 		}
