Allow control of Zen AutoLaunch process lifetime limit via an environment variable.

diff --git a/Engine/Source/Developer/Zen/Private/ZenServerInterface.cpp b/Engine/Source/Developer/Zen/Private/ZenServerInterface.cpp
--- a/Engine/Source/Developer/Zen/Private/ZenServerInterface.cpp
+++ b/Engine/Source/Developer/Zen/Private/ZenServerInterface.cpp
@@ -556,6 +556,16 @@ IsLocalHost(const FString& Host)
 	return false;
 }
 
+static void
+ApplyProcessLifetimeOverride(bool& bLimitProcessLifetime)
+{
+	FString LimitProcessLifetime = FPlatformMisc::GetEnvironmentVariable(TEXT("UE-ZenLimitProcessLifetime"));
+	if (!LimitProcessLifetime.IsEmpty())
+	{
+		bLimitProcessLifetime = LimitProcessLifetime.ToBool();
+	}
+}
+
 void
 FServiceSettings::ReadFromConfig()
 {
@@ -579,6 +589,7 @@ FServiceSettings::ReadFromConfig()
 			ReadUInt16FromConfig(AutoLaunchConfigSection, TEXT("DesiredPort"), AutoLaunchSettings.DesiredPort, GEngineIni);
 			GConfig->GetBool(AutoLaunchConfigSection, TEXT("ShowConsole"), AutoLaunchSettings.bShowConsole, GEngineIni);
 			GConfig->GetBool(AutoLaunchConfigSection, TEXT("LimitProcessLifetime"), AutoLaunchSettings.bLimitProcessLifetime, GEngineIni);
+			ApplyProcessLifetimeOverride(AutoLaunchSettings.bLimitProcessLifetime);
 			GConfig->GetBool(TEXT("/Script/UnrealEd.CrashReportsPrivacySettings"), TEXT("bSendUnattendedBugReports"), AutoLaunchSettings.bSendUnattendedBugReports, GEditorSettingsIni);
 		}
 	}
@@ -611,6 +622,7 @@ FServiceSettings::ReadFromCompactBinary(FCbFieldView Field)
 				AutoLaunchSettings.DesiredPort = AutoLaunchSettingsObject["DesiredPort"].AsInt16();
 				AutoLaunchSettings.bShowConsole = AutoLaunchSettingsObject["ShowConsole"].AsBool();
 				AutoLaunchSettings.bLimitProcessLifetime = AutoLaunchSettingsObject["LimitProcessLifetime"].AsBool();
+				ApplyProcessLifetimeOverride(AutoLaunchSettings.bLimitProcessLifetime);
 				AutoLaunchSettings.bSendUnattendedBugReports = AutoLaunchSettingsObject["SendUnattendedBugReports"].AsBool();
 			}
 		}
