Add CREATE_BREAKAWAY_FROM_JOB flag when launching zenserver to ensure that it can outlive the Horde build step in which it was spawned.

diff --git a/Engine/Source/Developer/Zen/Private/ZenServerInterface.cpp b/Engine/Source/Developer/Zen/Private/ZenServerInterface.cpp
--- a/Engine/Source/Developer/Zen/Private/ZenServerInterface.cpp
+++ b/Engine/Source/Developer/Zen/Private/ZenServerInterface.cpp
@@ -1289,7 +1289,7 @@ StartLocalService(const FZenLocalServiceRunContext& Context, const TCHAR* Transi
 
 		FString CommandLine = FString::Printf(TEXT("\"%s\" %s"), *PlatformExecutable, *Parms);
 		PROCESS_INFORMATION ProcInfo;
-		if (CreateProcess(NULL, CommandLine.GetCharArray().GetData(), nullptr, nullptr, false, (::DWORD)(NORMAL_PRIORITY_CLASS | DETACHED_PROCESS), nullptr, PlatformWorkingDirectory.GetCharArray().GetData(), &StartupInfo, &ProcInfo))
+		if (CreateProcess(NULL, CommandLine.GetCharArray().GetData(), nullptr, nullptr, false, (::DWORD)(NORMAL_PRIORITY_CLASS | DETACHED_PROCESS | CREATE_BREAKAWAY_FROM_JOB), nullptr, PlatformWorkingDirectory.GetCharArray().GetData(), &StartupInfo, &ProcInfo))
 		{
 			::CloseHandle(ProcInfo.hThread);
 			Proc = FProcHandle(ProcInfo.hProcess);
