Improved fallback logic and user feedback when datapath for zenserver is incorrectly configured
Validate that source files exists when trying to install zenserver

diff --git a/Engine/Source/Developer/Zen/Private/ZenServerInterface.cpp b/Engine/Source/Developer/Zen/Private/ZenServerInterface.cpp
--- a/Engine/Source/Developer/Zen/Private/ZenServerInterface.cpp
+++ b/Engine/Source/Developer/Zen/Private/ZenServerInterface.cpp
@@ -342,17 +342,39 @@ GetZenVersion(const FString& UtilityPath, const FString& ServicePath, const FStr
 	return ComparableVersion;
 }
 
+static void
+PromptUserToSyncInTreeVersion(const FString& ServerFilePath)
+{
+	if (FApp::IsUnattended())
+	{
+		// Just log as there is no one to show a message
+		UE_LOG(LogZenServiceInstance, Display, TEXT("ZenServer can not verify installation. Please make sure your source installation in properly synced at '%s'"), *FPaths::GetPath(ServerFilePath));
+		return;
+	}
+
+	FText ZenSyncSourcePromptTitle = NSLOCTEXT("Zen", "Zen_SyncSourcePromptTitle", "Failed to launch");
+	FText ZenSyncSourcePromptText = FText::Format(NSLOCTEXT("Zen", "Zen_SyncSourcePromptText", "ZenServer can not verify installation. Please make sure your source installation in properly synced at '{0}'"), FText::FromString(FPaths::GetPath(ServerFilePath)));
+	FPlatformMisc::MessageBoxExt(EAppMsgType::Ok, *ZenSyncSourcePromptTitle.ToString(), *ZenSyncSourcePromptText.ToString());
+}
+
 static bool
 IsInstallVersionOutOfDate(const FString& InTreeUtilityPath, const FString& InstallUtilityPath, const FString& InTreeServicePath, const FString& InstallServicePath, FString& OutInTreeVersionCache, FString& OutInstallVersionCache)
 {
 	OutInTreeVersionCache = FPaths::Combine(FPaths::EngineSavedDir(), TEXT("Zen"), TEXT("zen.version"));
 	OutInstallVersionCache = FPaths::SetExtension(InstallUtilityPath, TEXT("version"));
 
+	IFileManager& FileManager = IFileManager::Get();
+	if (!FileManager.FileExists(*InTreeUtilityPath) || !FileManager.FileExists(*InTreeServicePath))
+	{
+		UE_LOG(LogZenServiceInstance, Warning, TEXT("InTree version at '%s' is invalid"), *InTreeServicePath);
+		PromptUserToSyncInTreeVersion(InTreeServicePath);
+		return false;
+	}
+
 	// Always get the InTree utility path so cached version information is up to date
 	FZenVersion InTreeVersion = GetZenVersion(InTreeUtilityPath, InTreeServicePath, OutInTreeVersionCache);
 	UE_LOG(LogZenServiceInstance, Log, TEXT("InTree version at '%s' is '%s'"), *InTreeServicePath, *InTreeVersion.ToString());
 
-	IFileManager& FileManager = IFileManager::Get();
 	if (!FileManager.FileExists(*InstallUtilityPath) || !FileManager.FileExists(*InstallServicePath))
 	{
 		UE_LOG(LogZenServiceInstance, Log, TEXT("No installation found at '%s'"), *InstallServicePath);
@@ -461,54 +483,91 @@ DetermineLocalDataCachePath(const TCHAR* ConfigSection, FString& DataPath)
 	}
 }
 
-static void
-DetermineDataPath(const TCHAR* ConfigSection, FString& DataPath)
+static bool
+DetermineDataPath(const TCHAR* ConfigSection, FString& DataPath, bool& HasInvalidPathConfigurations)
+{
+	auto ValidateDataPath = [](const FString& InDataPath)
 {
-	auto NormalizeDataPath = [](const FString& InDataPath)
+		if (InDataPath.IsEmpty())
 	{
+			return FString{};
+		}
+		IFileManager& FileManager = IFileManager::Get();
 		FString FinalPath = FPaths::ConvertRelativePathToFull(InDataPath);
 		FPaths::NormalizeDirectoryName(FinalPath);
+		FFileStatData StatData = FileManager.GetStatData(*InDataPath);
+		if (StatData.bIsValid && StatData.bIsDirectory)
+		{
+			FString TestFilePath = FinalPath / ".zen-startup-test-file";
+			FArchive* TestFile = FileManager.CreateFileWriter(*TestFilePath, FILEWRITE_Silent);
+			if (!TestFile)
+			{
+				return FString{};
+			}
+			TestFile->Close();
+			delete TestFile;
+			FileManager.Delete(*TestFilePath);
 		return FinalPath;
+		}
+		if (FileManager.MakeDirectory(*InDataPath, true))
+		{
+			return FinalPath;
+		}
+		return FString{};
 	};
 
 	// Zen commandline
 	FString CommandLineOverrideValue;
-	if (FParse::Value(FCommandLine::Get(), TEXT("ZenDataPath="), CommandLineOverrideValue))
+	if (FParse::Value(FCommandLine::Get(), TEXT("ZenDataPath="), CommandLineOverrideValue) && !CommandLineOverrideValue.IsEmpty())
+	{
+		if (FString Path = ValidateDataPath(CommandLineOverrideValue); !Path.IsEmpty())
 	{
-		DataPath = NormalizeDataPath(CommandLineOverrideValue);
+			DataPath = Path;
 		UE_LOG(LogZenServiceInstance, Log, TEXT("Found command line override ZenDataPath=%s"), *CommandLineOverrideValue);
-		return;
+			return true;
+		}
+		UE_LOG(LogZenServiceInstance, Warning, TEXT("Skipping command line override ZenDataPath=%s due to an invalid path"), *CommandLineOverrideValue);
+		HasInvalidPathConfigurations = true;
 	}
 
 	// Zen subprocess environment
-	FString SubprocessDataPathEnvOverrideValue;
-	SubprocessDataPathEnvOverrideValue = FPlatformMisc::GetEnvironmentVariable(TEXT("UE-ZenSubprocessDataPath"));
-	if (!SubprocessDataPathEnvOverrideValue.IsEmpty())
+	if (FString SubprocessDataPathEnvOverrideValue = FPlatformMisc::GetEnvironmentVariable(TEXT("UE-ZenSubprocessDataPath")); !SubprocessDataPathEnvOverrideValue.IsEmpty())
+	{
+		if (FString Path = ValidateDataPath(SubprocessDataPathEnvOverrideValue); !Path.IsEmpty())
 	{
-		DataPath = NormalizeDataPath(SubprocessDataPathEnvOverrideValue);
+			DataPath = Path;
 		UE_LOG(LogZenServiceInstance, Log, TEXT("Found subprocess environment variable UE-ZenSubprocessDataPath=%s"), *SubprocessDataPathEnvOverrideValue);
-		return;
+			return true;
+		}
+		UE_LOG(LogZenServiceInstance, Warning, TEXT("Skipping subprocess environment variable UE-ZenSubprocessDataPath=%s due to an invalid path"), *SubprocessDataPathEnvOverrideValue);
+		HasInvalidPathConfigurations = true;
 	}
 
 	// Zen registry/stored
 	FString DataPathEnvOverrideValue;
-	if (FPlatformMisc::GetStoredValue(TEXT("Epic Games"), TEXT("Zen"), TEXT("DataPath"), DataPathEnvOverrideValue))
+	if (FPlatformMisc::GetStoredValue(TEXT("Epic Games"), TEXT("Zen"), TEXT("DataPath"), DataPathEnvOverrideValue) && !DataPathEnvOverrideValue.IsEmpty())
 	{
-		if (!DataPathEnvOverrideValue.IsEmpty())
+		if (FString Path = ValidateDataPath(DataPathEnvOverrideValue); !Path.IsEmpty())
 		{
-			DataPath = NormalizeDataPath(DataPathEnvOverrideValue);
-			UE_LOG(LogZenServiceInstance, Log, TEXT("Found registry key Zen DataPath=%s"), *DataPath);
-			return;
+			DataPath = Path;
+			UE_LOG(LogZenServiceInstance, Log, TEXT("Found registry key Zen DataPath=%s"), *DataPathEnvOverrideValue);
+			return true;
 		}
+		UE_LOG(LogZenServiceInstance, Warning, TEXT("Skipping registry key Zen DataPath=%s due to an invalid path"), *DataPathEnvOverrideValue);
+		HasInvalidPathConfigurations = true;
 	}
 
 	// Zen environment
-	DataPathEnvOverrideValue = FPlatformMisc::GetEnvironmentVariable(TEXT("UE-ZenDataPath"));
-	if (!DataPathEnvOverrideValue.IsEmpty())
+	if (FString ZenDataPathEnvOverrideValue = FPlatformMisc::GetEnvironmentVariable(TEXT("UE-ZenDataPath")); !ZenDataPathEnvOverrideValue.IsEmpty())
 	{
-		DataPath = NormalizeDataPath(DataPathEnvOverrideValue);
-		UE_LOG(LogZenServiceInstance, Log, TEXT("Found environment variable UE-ZenDataPath=%s"), *DataPathEnvOverrideValue);
-		return;
+		if (FString Path = ValidateDataPath(ZenDataPathEnvOverrideValue); !Path.IsEmpty())
+	{
+			DataPath = Path;
+			UE_LOG(LogZenServiceInstance, Log, TEXT("Found environment variable UE-ZenDataPath=%s"), *ZenDataPathEnvOverrideValue);
+			return true;
+		}
+		UE_LOG(LogZenServiceInstance, Warning, TEXT("Skipping environment variable UE-ZenDataPath=%s due to an invalid path"), *ZenDataPathEnvOverrideValue);
+		HasInvalidPathConfigurations = true;
 	}
 
 	// Follow local DDC (if outside workspace)
@@ -516,15 +575,32 @@ DetermineDataPath(const TCHAR* ConfigSection, FString& DataPath)
 	DetermineLocalDataCachePath(ConfigSection, LocalDataCachePath);
 	if (!LocalDataCachePath.IsEmpty() && (LocalDataCachePath != TEXT("None")) && !FPaths::IsUnderDirectory(LocalDataCachePath, FPaths::RootDir()))
 	{
-		DataPath = NormalizeDataPath(FPaths::Combine(LocalDataCachePath, TEXT("Zen")));
-		return;
+		if (FString Path = ValidateDataPath(LocalDataCachePath); !Path.IsEmpty())
+		{
+			DataPath = Path;
+			UE_LOG(LogZenServiceInstance, Log, TEXT("Found local data cache path=%s"), *LocalDataCachePath);
+			return true;
+		}
+		UE_LOG(LogZenServiceInstance, Warning, TEXT("Skipping local data cache path=%s due to an invalid path"), *LocalDataCachePath);
+		HasInvalidPathConfigurations = true;
 	}
 
 	// Zen config default
-	GConfig->GetString(ConfigSection, TEXT("DataPath"), DataPath, GEngineIni);
-	DataPath = NormalizeDataPath(DataPath);
-
-	check(!DataPath.IsEmpty())
+	FString ConfigDefaultPath;
+	GConfig->GetString(ConfigSection, TEXT("DataPath"), ConfigDefaultPath, GEngineIni);
+	if (!ConfigDefaultPath.IsEmpty())
+	{
+		if (FString Path = ValidateDataPath(ConfigDefaultPath); !Path.IsEmpty())
+		{
+			DataPath = Path;
+			UE_LOG(LogZenServiceInstance, Log, TEXT("Found Zen config default=%s"), *ConfigDefaultPath);
+			return true;
+		}
+		UE_LOG(LogZenServiceInstance, Warning, TEXT("Skipping Zen config default=%s due to an invalid path"), *ConfigDefaultPath);
+		HasInvalidPathConfigurations = true;
+	}
+	UE_LOG(LogZenServiceInstance, Warning, TEXT("Unable to determine a valid Zen data path"));
+	return false;
 }
 
 static void
@@ -566,6 +642,37 @@ ApplyProcessLifetimeOverride(bool& bLimitProcessLifetime)
 	}
 }
 
+static void
+PromptUserUnableToDetermineValidDataPath()
+{
+	if (FApp::IsUnattended())
+	{
+		// Just log as there is no one to show a message
+		UE_LOG(LogZenServiceInstance, Display, TEXT("ZenServer is unable to determine a valid data path"));
+		return;
+	}
+	FString LogDirPath = FPaths::ConvertRelativePathToFull(FPaths::ProjectLogDir());
+	FText ZenInvalidDataPathPromptTitle = NSLOCTEXT("Zen", "Zen_InvalidDataPathPromptTitle", "No Valid Data Path Configuration");
+	FText ZenInvalidDataPathPromptText = FText::Format(NSLOCTEXT("Zen", "Zen_InvalidDataPathPromptText", "ZenServer can not determine a valid data path.\nPlease check the log in '{0}' for details.\nUpdate your configuration and restart."), FText::FromString(LogDirPath));
+	FPlatformMisc::MessageBoxExt(EAppMsgType::Ok, *ZenInvalidDataPathPromptText.ToString(), *ZenInvalidDataPathPromptTitle.ToString());
+}
+
+static void
+PromptUserAboutInvalidValidDataPathConfiguration(const FString& UsedDataPath)
+{
+	if (FApp::IsUnattended())
+	{
+		// Just log as there is no one to show a message
+		UE_LOG(LogZenServiceInstance, Display, TEXT("ZenServer has detected invalid data path configuration. Falling back to '%s'"), *UsedDataPath);
+		return;
+	}
+
+	FString LogDirPath = FPaths::ConvertRelativePathToFull(FPaths::ProjectLogDir());
+	FText ZenInvalidValidDataPathConfigurationPromptTitle = NSLOCTEXT("Zen", "Zen_InvalidValidDataPathConfigurationPromptTitle", "Invalid Data Paths");
+	FText ZenInvalidValidDataPathConfigurationPromptText = FText::Format(NSLOCTEXT("Zen", "Zen_InvalidValidDataPathConfigurationPromptText", "ZenServer has detected invalid data path configuration.\nPlease check the log in '{0}' for details.\n\nFalling back to using '{1}' as data path."), FText::FromString(LogDirPath), FText::FromString(UsedDataPath));
+	FPlatformMisc::MessageBoxExt(EAppMsgType::Ok, *ZenInvalidValidDataPathConfigurationPromptText.ToString(), *ZenInvalidValidDataPathConfigurationPromptTitle.ToString());
+}
+
 void
 FServiceSettings::ReadFromConfig()
 {
@@ -583,7 +690,17 @@ FServiceSettings::ReadFromConfig()
 			SettingsVariant.Emplace<FServiceAutoLaunchSettings>();
 			FServiceAutoLaunchSettings& AutoLaunchSettings = SettingsVariant.Get<FServiceAutoLaunchSettings>();
 
-			DetermineDataPath(AutoLaunchConfigSection, AutoLaunchSettings.DataPath);
+			bool HasInvalidPathConfigurations = false;
+			if (!DetermineDataPath(AutoLaunchConfigSection, AutoLaunchSettings.DataPath, HasInvalidPathConfigurations))
+			{
+				PromptUserUnableToDetermineValidDataPath();
+				FPlatformMisc::RequestExit(true);
+				return;
+			}
+			else if (HasInvalidPathConfigurations)
+			{
+				PromptUserAboutInvalidValidDataPathConfiguration(AutoLaunchSettings.DataPath);
+			}
 			GConfig->GetString(AutoLaunchConfigSection, TEXT("ExtraArgs"), AutoLaunchSettings.ExtraArgs, GEngineIni);
 
 			ReadUInt16FromConfig(AutoLaunchConfigSection, TEXT("DesiredPort"), AutoLaunchSettings.DesiredPort, GEngineIni);
