Make sure zen data path is in Zen folder when following local DDC path

diff --git a/Engine/Source/Developer/Zen/Private/ZenServerInterface.cpp b/Engine/Source/Developer/Zen/Private/ZenServerInterface.cpp
--- a/Engine/Source/Developer/Zen/Private/ZenServerInterface.cpp
+++ b/Engine/Source/Developer/Zen/Private/ZenServerInterface.cpp
@@ -575,9 +575,10 @@ DetermineDataPath(const TCHAR* ConfigSection, FString& DataPath, bool& HasInvali
 	DetermineLocalDataCachePath(ConfigSection, LocalDataCachePath);
 	if (!LocalDataCachePath.IsEmpty() && (LocalDataCachePath != TEXT("None")) && !FPaths::IsUnderDirectory(LocalDataCachePath, FPaths::RootDir()))
 	{
-		if (FString Path = ValidateDataPath(LocalDataCachePath); !Path.IsEmpty())
+		FString ZenLocalDataCachePath = FPaths::Combine(LocalDataCachePath, TEXT("Zen"));
+		if (FString Path = ValidateDataPath(ZenLocalDataCachePath); !Path.IsEmpty())
 		{
-			DataPath = Path;
+			DataPath = ZenLocalDataCachePath;
 			UE_LOG(LogZenServiceInstance, Log, TEXT("Found local data cache path=%s"), *LocalDataCachePath);
 			return true;
 		}
