Fix a collection of bugs related to ZenServer launching:
- Fix a message that had the text and caption fields inverted.
- Ensure that sandbox instances don't overwrite the .runcontext file used by non-sandbox instances.
- Fix a bug where the Zen data path when following DDC path settings wasn't normalizing the path leading to needless shutdown and startup of zenserver.

diff --git a/Engine/Source/Developer/DerivedDataCache/Private/ZenCacheStore.cpp b/Engine/Source/Developer/DerivedDataCache/Private/ZenCacheStore.cpp
--- a/Engine/Source/Developer/DerivedDataCache/Private/ZenCacheStore.cpp
+++ b/Engine/Source/Developer/DerivedDataCache/Private/ZenCacheStore.cpp
@@ -1490,6 +1490,7 @@ ILegacyCacheStore* CreateZenCacheStore(const TCHAR* NodeName, const TCHAR* Confi
 		FPaths::NormalizeDirectoryName(AutoLaunchSettings.DataPath);
 		AutoLaunchSettings.DataPath += TEXT("_");
 		AutoLaunchSettings.DataPath += Sandbox;
+		AutoLaunchSettings.bIsDefaultSharedRunContext = false;
 
 		// The unique local instances will always limit process lifetime for now to avoid accumulating many of them
 		AutoLaunchSettings.bLimitProcessLifetime = true;
diff --git a/Engine/Source/Developer/Zen/Private/ZenServerInterface.cpp b/Engine/Source/Developer/Zen/Private/ZenServerInterface.cpp
--- a/Engine/Source/Developer/Zen/Private/ZenServerInterface.cpp
+++ b/Engine/Source/Developer/Zen/Private/ZenServerInterface.cpp
@@ -354,7 +354,7 @@ PromptUserToSyncInTreeVersion(const FString& ServerFilePath)
 
 	FText ZenSyncSourcePromptTitle = NSLOCTEXT("Zen", "Zen_SyncSourcePromptTitle", "Failed to launch");
 	FText ZenSyncSourcePromptText = FText::Format(NSLOCTEXT("Zen", "Zen_SyncSourcePromptText", "ZenServer can not verify installation. Please make sure your source installation in properly synced at '{0}'"), FText::FromString(FPaths::GetPath(ServerFilePath)));
-	FPlatformMisc::MessageBoxExt(EAppMsgType::Ok, *ZenSyncSourcePromptTitle.ToString(), *ZenSyncSourcePromptText.ToString());
+	FPlatformMisc::MessageBoxExt(EAppMsgType::Ok, *ZenSyncSourcePromptText.ToString(), *ZenSyncSourcePromptTitle.ToString());
 }
 
 static bool
@@ -578,7 +578,7 @@ DetermineDataPath(const TCHAR* ConfigSection, FString& DataPath, bool& HasInvali
 		FString ZenLocalDataCachePath = FPaths::Combine(LocalDataCachePath, TEXT("Zen"));
 		if (FString Path = ValidateDataPath(ZenLocalDataCachePath); !Path.IsEmpty())
 		{
-			DataPath = ZenLocalDataCachePath;
+			DataPath = Path;
 			UE_LOG(LogZenServiceInstance, Log, TEXT("Found local data cache path=%s"), *LocalDataCachePath);
 			return true;
 		}
@@ -742,6 +742,7 @@ FServiceSettings::ReadFromCompactBinary(FCbFieldView Field)
 				AutoLaunchSettings.bLimitProcessLifetime = AutoLaunchSettingsObject["LimitProcessLifetime"].AsBool();
 				ApplyProcessLifetimeOverride(AutoLaunchSettings.bLimitProcessLifetime);
 				AutoLaunchSettings.bSendUnattendedBugReports = AutoLaunchSettingsObject["SendUnattendedBugReports"].AsBool();
+				AutoLaunchSettings.bIsDefaultSharedRunContext = AutoLaunchSettingsObject["IsDefaultSharedRunContext"].AsBool(AutoLaunchSettings.bIsDefaultSharedRunContext);
 			}
 		}
 	}
@@ -798,6 +799,7 @@ FServiceSettings::WriteToCompactBinary(FCbWriter& Writer) const
 		Writer << "ShowConsole" << AutoLaunchSettings.bShowConsole;
 		Writer << "LimitProcessLifetime" << AutoLaunchSettings.bLimitProcessLifetime;
 		Writer << "SendUnattendedBugReports" << AutoLaunchSettings.bSendUnattendedBugReports;
+		Writer << "IsDefaultSharedRunContext" << AutoLaunchSettings.bIsDefaultSharedRunContext;
 		Writer.EndObject();
 	}
 	else
@@ -1847,6 +1849,7 @@ FZenServiceInstance::AutoLaunch(const FServiceAutoLaunchSettings& InSettings, FS
 			break;
 		}
 		FPlatformProcess::Sleep(0.1f);
+		LockFileState = LockFileData();
 	}
 
 	uint16 ShutdownEffectivePort = 0;
@@ -1862,23 +1865,30 @@ FZenServiceInstance::AutoLaunch(const FServiceAutoLaunchSettings& InSettings, FS
 			return false;
 		}
 
-		FZenLocalServiceRunContext DesiredRunContext;
-		DesiredRunContext.Executable = ExecutablePath;
-		DesiredRunContext.CommandlineArguments = DetermineCmdLineWithoutTransientComponents(InSettings, InSettings.DesiredPort);
-		DesiredRunContext.WorkingDirectory = WorkingDirectory;
-		DesiredRunContext.DataPath = InSettings.DataPath;
-		DesiredRunContext.bShowConsole = InSettings.bShowConsole;
-
-		FZenLocalServiceRunContext CurrentRunContext;
-		if (CurrentRunContext.ReadFromJsonFile(*ExecutionContextFilePath) && (DesiredRunContext == CurrentRunContext))
+		if (InSettings.bIsDefaultSharedRunContext)
 		{
-			UE_LOG(LogZenServiceInstance, Log, TEXT("Found existing instance running on port %u matching our settings"), InSettings.DesiredPort);
-			bLaunchNewInstance = false;
+			FZenLocalServiceRunContext DesiredRunContext;
+			DesiredRunContext.Executable = ExecutablePath;
+			DesiredRunContext.CommandlineArguments = DetermineCmdLineWithoutTransientComponents(InSettings, InSettings.DesiredPort);
+			DesiredRunContext.WorkingDirectory = WorkingDirectory;
+			DesiredRunContext.DataPath = InSettings.DataPath;
+			DesiredRunContext.bShowConsole = InSettings.bShowConsole;
+
+			FZenLocalServiceRunContext CurrentRunContext;
+			if (CurrentRunContext.ReadFromJsonFile(*ExecutionContextFilePath) && (DesiredRunContext == CurrentRunContext))
+			{
+				UE_LOG(LogZenServiceInstance, Log, TEXT("Found existing instance running on port %u matching our settings"), InSettings.DesiredPort);
+				bLaunchNewInstance = false;
+			}
+			else
+			{
+				UE_LOG(LogZenServiceInstance, Log, TEXT("Found existing instance running on port %u with different run context, will attempt shut down"), InSettings.DesiredPort);
+				ShutdownEffectivePort = LockFileState.EffectivePort;
+			}
 		}
 		else
 		{
-			UE_LOG(LogZenServiceInstance, Log, TEXT("Found existing instance running on port %u with different run context, will attempt shut down"), InSettings.DesiredPort);
-			ShutdownEffectivePort = LockFileState.EffectivePort;
+			UE_LOG(LogZenServiceInstance, Log, TEXT("Found existing instance running on port %u when not using shared context, will use it"), InSettings.DesiredPort);
 		}
 	}
 	else
@@ -1931,8 +1941,8 @@ FZenServiceInstance::AutoLaunch(const FServiceAutoLaunchSettings& InSettings, FS
 
 		FProcHandle Proc = StartLocalService(EffectiveRunContext, TransientParms.IsEmpty() ? nullptr : *TransientParms);
 
-		// Don't write run context unless we actually started a new instance
-		if (bLaunchNewInstance)
+		// Only write run context if we're using the default shared run context
+		if (InSettings.bIsDefaultSharedRunContext)
 		{
 			EffectiveRunContext.WriteToJsonFile(*ExecutionContextFilePath);
 		}
diff --git a/Engine/Source/Developer/Zen/Public/Experimental/ZenServerInterface.h b/Engine/Source/Developer/Zen/Public/Experimental/ZenServerInterface.h
--- a/Engine/Source/Developer/Zen/Public/Experimental/ZenServerInterface.h
+++ b/Engine/Source/Developer/Zen/Public/Experimental/ZenServerInterface.h
@@ -43,6 +43,7 @@ struct FServiceAutoLaunchSettings
 	bool bShowConsole = false;
 	bool bLimitProcessLifetime = false;
 	bool bSendUnattendedBugReports = false;
+	bool bIsDefaultSharedRunContext = true;
 };
 
 struct FServiceSettings
