Ensure that ZenServer can perform recovery successfully if terminated during a multi-process cook.  Uses FSystemWideCriticalSection to coordinate recovery operation across processes.

diff --git a/Engine/Source/Developer/Zen/Private/ZenServerInterface.cpp b/Engine/Source/Developer/Zen/Private/ZenServerInterface.cpp
--- a/Engine/Source/Developer/Zen/Private/ZenServerInterface.cpp
+++ b/Engine/Source/Developer/Zen/Private/ZenServerInterface.cpp
@@ -83,12 +83,15 @@ public:
 	};
 	static_assert(sizeof(ZenServerEntry) == 64);
 
-	ZenServerEntry* LookupByDesiredListenPort(int DesiredListenPort);
 	const ZenServerEntry* LookupByDesiredListenPort(int DesiredListenPort) const;
-	const ZenServerEntry* LookupByEffectiveListenPort(int DesiredListenPort) const;
+	ZenServerEntry* LookupByDesiredListenPort(int DesiredListenPort);
+	const ZenServerEntry* LookupByEffectiveListenPort(int EffectiveListenPort) const;
+	ZenServerEntry* LookupByEffectiveListenPort(int EffectiveListenPort);
 	const ZenServerEntry* LookupByPid(uint32 Pid) const;
 
 private:
+	const ZenServerEntry* LookupByDesiredListenPortInternal(int DesiredListenPort) const;
+	const ZenServerEntry* LookupByEffectiveListenPortInternal(int EffectiveListenPort) const;
 	void* m_hMapFile = nullptr;
 	ZenServerEntry* m_Data = nullptr;
 	int				m_MaxEntryCount = 65536 / sizeof(ZenServerEntry);
@@ -174,7 +177,7 @@ ZenServerState::~ZenServerState()
 	m_Data = nullptr;
 }
 
-const ZenServerState::ZenServerEntry* ZenServerState::LookupByDesiredListenPort(int Port) const
+const ZenServerState::ZenServerEntry* ZenServerState::LookupByDesiredListenPortInternal(int Port) const
 {
 	if (m_Data == nullptr)
 	{
@@ -192,10 +195,19 @@ const ZenServerState::ZenServerEntry* ZenServerState::LookupByDesiredListenPort(
 	return nullptr;
 }
 
+const ZenServerState::ZenServerEntry* ZenServerState::LookupByDesiredListenPort(int Port) const
+{
+	return LookupByDesiredListenPortInternal(Port);
+}
+
 ZenServerState::ZenServerEntry* ZenServerState::LookupByDesiredListenPort(int Port)
 {
 	check(!m_IsReadOnly);
+	return const_cast<ZenServerState::ZenServerEntry*>(LookupByDesiredListenPortInternal(Port));
+}
 
+const ZenServerState::ZenServerEntry* ZenServerState::LookupByEffectiveListenPortInternal(int Port) const
+{
 	if (m_Data == nullptr)
 	{
 		return nullptr;
@@ -203,7 +215,7 @@ ZenServerState::ZenServerEntry* ZenServerState::LookupByDesiredListenPort(int Po
 
 	for (int i = 0; i < m_MaxEntryCount; ++i)
 	{
-		if (m_Data[i].DesiredListenPort == Port)
+		if (m_Data[i].EffectiveListenPort == Port)
 		{
 			return &m_Data[i];
 		}
@@ -214,20 +226,13 @@ ZenServerState::ZenServerEntry* ZenServerState::LookupByDesiredListenPort(int Po
 
 const ZenServerState::ZenServerEntry* ZenServerState::LookupByEffectiveListenPort(int Port) const
 {
-	if (m_Data == nullptr)
-	{
-		return nullptr;
-	}
-
-	for (int i = 0; i < m_MaxEntryCount; ++i)
-	{
-		if (m_Data[i].EffectiveListenPort == Port)
-		{
-			return &m_Data[i];
-		}
-	}
+	return LookupByEffectiveListenPortInternal(Port);
+}
 
-	return nullptr;
+ZenServerState::ZenServerEntry* ZenServerState::LookupByEffectiveListenPort(int Port)
+{
+	check(!m_IsReadOnly);
+	return const_cast<ZenServerState::ZenServerEntry*>(LookupByEffectiveListenPortInternal(Port));
 }
 
 const ZenServerState::ZenServerEntry* ZenServerState::LookupByPid(uint32 Pid) const
@@ -852,6 +857,7 @@ FServiceSettings::TryApplyAutoLaunchOverride()
 #if UE_WITH_ZEN
 
 uint16 FZenServiceInstance::AutoLaunchedPort = 0;
+uint32 FZenServiceInstance::AutoLaunchedPid = 0;
 
 struct LockFileData
 {
@@ -1628,9 +1634,9 @@ FZenServiceInstance::TryRecovery()
 		return false;
 	}
 
-	static FCriticalSection RecoveryCriticalSection;
 	static std::atomic<int64> LastRecoveryTicks;
 	static bool bLastRecoveryResult = false;
+	const FTimespan MaximumWaitForLaunch = FTimespan::FromSeconds(30);
 	const FTimespan MaximumWaitForHealth = FTimespan::FromSeconds(30);
 	const FTimespan MinimumDurationSinceLastRecovery = FTimespan::FromMinutes(2);
 
@@ -1638,32 +1644,43 @@ FZenServiceInstance::TryRecovery()
 
 	if (TimespanSinceLastRecovery > MinimumDurationSinceLastRecovery)
 	{
-		FScopeLock Lock(&RecoveryCriticalSection);
+		static FSystemWideCriticalSection RecoveryCriticalSection(TEXT("ZenServerRecovery"));
 		// Update timespan since it may have changed since we waited to enter the crit section
 		TimespanSinceLastRecovery = FDateTime::UtcNow() - FDateTime(LastRecoveryTicks.load(std::memory_order_relaxed));
 		if (TimespanSinceLastRecovery > MinimumDurationSinceLastRecovery)
 		{
 			UE_LOG(LogZenServiceInstance, Display, TEXT("Local ZenServer recovery being attempted..."));
-			FZenLocalServiceRunContext RunContext;
-			if (TryGetLocalServiceRunContext(RunContext))
+
+			bool bShutdownExistingInstance = true;
+			std::atomic<uint32> PreviousSponsorPids[UE_ARRAY_COUNT(ZenServerState::ZenServerEntry::SponsorPids)];
 			{
-				if (!ShutdownRunningServiceUsingExecutablePath(*RunContext.GetExecutable()))
+				const ZenServerState ServerState;
+				const ZenServerState::ZenServerEntry* Entry = ServerState.LookupByEffectiveListenPort(Port);
+				if (Entry)
 				{
-					return false;
+					if (Entry->Pid.load(std::memory_order_relaxed) != AutoLaunchedPid)
+					{
+						// The running process pid is not the same as the one we launched.  The process was relaunched elsewhere. Avoid shutting it down again.
+						bShutdownExistingInstance = false;
+					}
+
+					for (int32 SponsorPidIndex = 0; SponsorPidIndex < UE_ARRAY_COUNT(ZenServerState::ZenServerEntry::SponsorPids); ++SponsorPidIndex)
+					{
+						PreviousSponsorPids[SponsorPidIndex].store(Entry->SponsorPids[SponsorPidIndex].load(std::memory_order_relaxed), std::memory_order_relaxed);
+					}
 				}
-				StartLocalService(RunContext);
-				UE_LOG(LogZenServiceInstance, Display, TEXT("Local ZenServer recovery finished."));
 			}
-			else
+			if (bShutdownExistingInstance && !ShutdownRunningServiceUsingEffectivePort(Port))
 			{
-				UE_LOG(LogZenServiceInstance, Warning, TEXT("Local ZenServer recovery failed due to lack of run context."));
+				return false;
 			}
-			
+
+			AutoLaunch(Settings.SettingsVariant.Get<FServiceAutoLaunchSettings>(), *GetLocalServiceInstallPath(), HostName, Port);
 			FDateTime StartedWaitingForHealth = FDateTime::UtcNow();
 			bLastRecoveryResult = IsServiceReady();
 			while (!bLastRecoveryResult)
 			{
-				FTimespan WaitForHealth = StartedWaitingForHealth - FDateTime::UtcNow();
+				FTimespan WaitForHealth = FDateTime::UtcNow() - StartedWaitingForHealth;
 				if (WaitForHealth > MaximumWaitForHealth)
 				{
 					UE_LOG(LogZenServiceInstance, Warning, TEXT("Local ZenServer recovery timed out waiting for service to become healthy"));
@@ -1671,9 +1688,14 @@ FZenServiceInstance::TryRecovery()
 				}
 
 				FPlatformProcess::Sleep(0.5f);
+				if (!IsZenProcessUsingEffectivePort(Port))
+				{
+					AutoLaunch(Settings.SettingsVariant.Get<FServiceAutoLaunchSettings>(), *GetLocalServiceInstallPath(), HostName, Port);
+				}
 				bLastRecoveryResult = IsServiceReady();
 			}
 			LastRecoveryTicks.store(FDateTime::UtcNow().GetTicks(), std::memory_order_relaxed);
+			UE_LOG(LogZenServiceInstance, Display, TEXT("Local ZenServer recovery finished."));
 			
 			if (bLastRecoveryResult)
 			{
@@ -1683,6 +1705,16 @@ FZenServiceInstance::TryRecovery()
 			{
 				UE_LOG(LogZenServiceInstance, Display, TEXT("Local ZenServer post recovery status: NOT healthy"));
 			}
+
+			ZenServerState PostRecoveryServerState(/*bReadOnly*/ false);
+			ZenServerState::ZenServerEntry* PostRecoveryEntry = PostRecoveryServerState.LookupByEffectiveListenPort(Port);
+			if (PostRecoveryEntry)
+			{
+				for (int32 SponsorPidIndex = 0; SponsorPidIndex < UE_ARRAY_COUNT(ZenServerState::ZenServerEntry::SponsorPids); ++SponsorPidIndex)
+				{
+					PostRecoveryEntry->SponsorPids[SponsorPidIndex].store(PreviousSponsorPids[SponsorPidIndex].load(std::memory_order_relaxed), std::memory_order_relaxed);
+				}
+			}
 		}
 	}
 
@@ -1706,6 +1738,12 @@ FZenServiceInstance::Initialize()
 			bHasLaunchedLocal = AutoLaunch(Settings.SettingsVariant.Get<FServiceAutoLaunchSettings>(), *ExecutableInstallPath, HostName, Port);
 			if (bHasLaunchedLocal)
 			{
+				const ZenServerState State;
+				const ZenServerState::ZenServerEntry* RunningEntry = State.LookupByEffectiveListenPort(Port);
+				if (RunningEntry != nullptr)
+				{
+					AutoLaunchedPid = RunningEntry->Pid.load(std::memory_order_relaxed);
+				}
 				AutoLaunchedPort = Port;
 				bIsRunningLocally = true;
 			}
diff --git a/Engine/Source/Developer/Zen/Public/Experimental/ZenServerInterface.h b/Engine/Source/Developer/Zen/Public/Experimental/ZenServerInterface.h
--- a/Engine/Source/Developer/Zen/Public/Experimental/ZenServerInterface.h
+++ b/Engine/Source/Developer/Zen/Public/Experimental/ZenServerInterface.h
@@ -216,6 +216,7 @@ private:
 	FString HostName;
 	uint16 Port;
 	static uint16 AutoLaunchedPort;
+	static uint32 AutoLaunchedPid;
 	bool bHasLaunchedLocal = false;
 	bool bIsRunningLocally = true;
 };
