Fix missing include in ZenServerHttp.cpp submitted in CL 26841365

diff --git a/Engine/Source/Developer/Zen/Private/ZenServerHttp.cpp b/Engine/Source/Developer/Zen/Private/ZenServerHttp.cpp
--- a/Engine/Source/Developer/Zen/Private/ZenServerHttp.cpp
+++ b/Engine/Source/Developer/Zen/Private/ZenServerHttp.cpp
@@ -29,6 +29,7 @@
 #include "Containers/StringFwd.h"
 #include "HAL/LowLevelMemTracker.h"
 #include "Memory/CompositeBuffer.h"
+#include "Misc/App.h"
 #include "Serialization/CompactBinary.h"
 #include "Serialization/CompactBinaryPackage.h"
 #include "Serialization/CompactBinaryValidation.h"
