Increase default timeout when waiting for zenserver shutdown from 5 seconds to 15 seconds as I've received a report of zenserver upgrade failing due to the shutdown timeout being insufficient.

diff --git a/Engine/Source/Developer/Zen/Private/ZenServerInterface.cpp b/Engine/Source/Developer/Zen/Private/ZenServerInterface.cpp
--- a/Engine/Source/Developer/Zen/Private/ZenServerInterface.cpp
+++ b/Engine/Source/Developer/Zen/Private/ZenServerInterface.cpp
@@ -1112,7 +1112,7 @@ RequestZenShutdownOnEffectivePort(uint16 EffectiveListenPort)
 }
 
 static bool
-ShutdownRunningServiceUsingExecutablePath(const TCHAR* ExecutablePath, double MaximumWaitDurationSeconds = 5.0)
+ShutdownRunningServiceUsingExecutablePath(const TCHAR* ExecutablePath, double MaximumWaitDurationSeconds = 15.0)
 {
 	uint32 ServicePid;
 	if (!IsZenProcessActive(ExecutablePath, &ServicePid))
@@ -1168,7 +1168,7 @@ ShutdownRunningServiceUsingExecutablePath(const TCHAR* ExecutablePath, double Ma
 }
 
 static bool
-ShutdownRunningServiceUsingEffectivePort(uint16 EffectiveListenPort, double MaximumWaitDurationSeconds = 5.0)
+ShutdownRunningServiceUsingEffectivePort(uint16 EffectiveListenPort, double MaximumWaitDurationSeconds = 15.0)
 {
 	if (!IsZenProcessUsingEffectivePort(EffectiveListenPort))
 	{
diff --git a/Engine/Source/Developer/Zen/Public/Experimental/ZenServerInterface.h b/Engine/Source/Developer/Zen/Public/Experimental/ZenServerInterface.h
--- a/Engine/Source/Developer/Zen/Public/Experimental/ZenServerInterface.h
+++ b/Engine/Source/Developer/Zen/Public/Experimental/ZenServerInterface.h
@@ -114,7 +114,7 @@ private:
 
 UE_API bool IsLocalServiceRunning(const TCHAR* DataPath, uint16* OutPort = nullptr);
 UE_API FProcHandle StartLocalService(const FZenLocalServiceRunContext& Context, const TCHAR* TransientArgs = nullptr);
-UE_API bool StopLocalService(const TCHAR* DataPath, double MaximumWaitDurationSeconds = 5.0);
+UE_API bool StopLocalService(const TCHAR* DataPath, double MaximumWaitDurationSeconds = 15.0);
 
 UE_API FString GetLocalServiceInstallPath();
 UE_API FString GetLocalServiceInstallVersion(bool bDetailed = true);
