Fix zenserver instance shared memory state access flags.

#prefilght 64f0583fd140152917e66ef2

diff --git a/Engine/Source/Developer/Zen/Private/ZenServerInterface.cpp b/Engine/Source/Developer/Zen/Private/ZenServerInterface.cpp
--- a/Engine/Source/Developer/Zen/Private/ZenServerInterface.cpp
+++ b/Engine/Source/Developer/Zen/Private/ZenServerInterface.cpp
@@ -128,15 +128,16 @@ ZenServerState::ZenServerState(bool ReadOnly)
 		return;
 	}
 #elif PLATFORM_UNIX || PLATFORM_MAC
-	int Fd = shm_open("/UnrealEngineZen", O_RDONLY | O_CLOEXEC, 0666);
+	int OFlag = ReadOnly ? (O_RDONLY | O_CLOEXEC) : (O_RDWR | O_CREAT | O_CLOEXEC);
+	int Fd = shm_open("/UnrealEngineZen", OFlag, 0666);
 	if (Fd < 0)
 	{
 		return;
 	}
 	void* hMap = (void*)intptr_t(Fd);
 
-	int Prot = ReadOnly ? (PROT_WRITE | PROT_READ) : PROT_READ;
-	void* pBuf = mmap(nullptr, MapSize, Prot, MAP_PRIVATE, Fd, 0);
+	int Prot = ReadOnly ? PROT_READ : (PROT_WRITE | PROT_READ);
+	void* pBuf = mmap(nullptr, MapSize, Prot, MAP_SHARED, Fd, 0);
 	if (pBuf == MAP_FAILED)
 	{
 		close(Fd);
