ZenDashboard refresh to add project stats and clean up reporting of disk size.

diff --git a/Engine/Source/Developer/DerivedDataCache/Private/ZenCacheStore.cpp b/Engine/Source/Developer/DerivedDataCache/Private/ZenCacheStore.cpp
--- a/Engine/Source/Developer/DerivedDataCache/Private/ZenCacheStore.cpp
+++ b/Engine/Source/Developer/DerivedDataCache/Private/ZenCacheStore.cpp
@@ -1302,8 +1302,8 @@ void FZenCacheStore::Initialize(
 		}
 
 		// Issue a request for stats as it will be fetched asynchronously and issuing now makes them available sooner for future callers.
-		Zen::FZenStats ZenStats;
-		ZenService.GetInstance().GetStats(ZenStats);
+		Zen::FZenCacheStats ZenStats;
+		ZenService.GetInstance().GetCacheStats(ZenStats);
 	}
 
 	GConfig->GetInt(TEXT("Zen"), TEXT("BatchPutMaxBytes"), BatchPutMaxBytes, GEngineIni);
diff --git a/Engine/Source/Developer/Zen/Private/ZenModule.cpp b/Engine/Source/Developer/Zen/Private/ZenModule.cpp
new file mode 100644
--- /dev/null
+++ b/Engine/Source/Developer/Zen/Private/ZenModule.cpp
@@ -0,0 +1,16 @@
+// Copyright Epic Games, Inc. All Rights Reserved.
+
+#include "Modules/ModuleManager.h"
+
+class FZenModule
+	: public IModuleInterface
+{
+public:
+	// IModuleInterface interface
+	virtual void StartupModule() override
+	{
+		FModuleManager::Get().LoadModuleChecked(TEXT("Sockets"));
+	}
+};
+
+IMPLEMENT_MODULE(FDefaultModuleImpl, Zen)
diff --git a/Engine/Source/Developer/Zen/Private/ZenServerInterface.cpp b/Engine/Source/Developer/Zen/Private/ZenServerInterface.cpp
--- a/Engine/Source/Developer/Zen/Private/ZenServerInterface.cpp
+++ b/Engine/Source/Developer/Zen/Private/ZenServerInterface.cpp
@@ -679,6 +679,28 @@ PromptUserAboutInvalidValidDataPathConfiguration(const FString& UsedDataPath)
 	FPlatformMisc::MessageBoxExt(EAppMsgType::Ok, *ZenInvalidValidDataPathConfigurationPromptText.ToString(), *ZenInvalidValidDataPathConfigurationPromptTitle.ToString());
 }
 
+static void ReadCbField(FCbFieldView Field, UE::Zen::FZenSizeStats& OutValue)
+{
+	FCbObjectView ObjectView = Field.AsObjectView();
+	OutValue.Disk = ObjectView["disk"].AsDouble();
+	OutValue.Memory = ObjectView["memory"].AsDouble();
+}
+
+static void ReadCbField(FCbFieldView Field, UE::Zen::FZenCIDSizeStats& OutValue)
+{
+	FCbObjectView ObjectView = Field.AsObjectView();
+	OutValue.Tiny = ObjectView["tiny"].AsInt64();
+	OutValue.Small = ObjectView["small"].AsInt64();
+	OutValue.Large = ObjectView["large"].AsInt64();
+	OutValue.Total = ObjectView["total"].AsInt64();
+}
+
+static void ReadCbField(FCbFieldView Field, UE::Zen::FZenCIDStats& OutValue)
+{
+	FCbObjectView ObjectView = Field.AsObjectView();
+	ReadCbField(ObjectView["size"], OutValue.Size);
+}
+
 void
 FServiceSettings::ReadFromConfig()
 {
@@ -2113,140 +2135,121 @@ FZenServiceInstance::AutoLaunch(const FServiceAutoLaunchSettings& InSettings, FS
 }
 
 bool 
-FZenServiceInstance::GetStats(FZenStats& Stats)
+FZenServiceInstance::GetCacheStats(FZenCacheStats& Stats)
 {
 	check(IsInGameThread());
 
-	// If we've already requested a stats and they are ready then grab them
-	if ( StatsRequest.IsReady() == true )
+	// If we've already requested stats and they are ready then grab them
+	if ( CacheStatsRequest.IsReady() == true )
 	{
-		LastStats		= StatsRequest.Get();
-		LastStatsTime	= FPlatformTime::Cycles64();
+		LastCacheStats		= CacheStatsRequest.Get();
+		LastCacheStatsTime	= FPlatformTime::Cycles64();
 
-		StatsRequest.Reset();
+		CacheStatsRequest.Reset();
 	}
 	
 	// Make a copy of the last updated stats
-	Stats = LastStats;
+	Stats = LastCacheStats;
 
 	const uint64 CurrentTime = FPlatformTime::Cycles64();
 	constexpr double MinTimeBetweenRequestsInSeconds = 0.5;
-	const double DeltaTimeInSeconds = FPlatformTime::ToSeconds64(CurrentTime - LastStatsTime);
+	const double DeltaTimeInSeconds = FPlatformTime::ToSeconds64(CurrentTime - LastCacheStatsTime);
 
-	if (!StatsRequest.IsValid() && DeltaTimeInSeconds > MinTimeBetweenRequestsInSeconds)
+	if (!CacheStatsRequest.IsValid() && DeltaTimeInSeconds > MinTimeBetweenRequestsInSeconds)
 	{
 #if WITH_EDITOR
 		EAsyncExecution ThreadPool = EAsyncExecution::LargeThreadPool;
 #else
 		EAsyncExecution ThreadPool = EAsyncExecution::ThreadPool;
 #endif
-		if (!StatsHttpRequest.IsValid())
+		if (!CacheStatsHttpRequest.IsValid())
 		{
 			TStringBuilder<128> ZenDomain;
 			ZenDomain << HostName << TEXT(":") << Port;
-			StatsHttpRequest = MakePimpl<FZenHttpRequest>(ZenDomain.ToString(), false);
+			CacheStatsHttpRequest = MakePimpl<FZenHttpRequest>(ZenDomain.ToString(), false);
 		}
 
 		// We've not got any requests in flight and we've met a given time requirement for requests
-		StatsRequest = Async(ThreadPool, [this]
-			{
-				UE::Zen::FZenHttpRequest& Request = *StatsHttpRequest.Get();
-				Request.Reset();
-
-				TArray64<uint8> GetBuffer;
-				FZenHttpRequest::Result Result = Request.PerformBlockingDownload(TEXTVIEW("/stats/z$"), &GetBuffer, Zen::EContentType::CbObject);
-
-				FZenStats Stats;
-
-				if (Result == Zen::FZenHttpRequest::Result::Success && Request.GetResponseCode() == 200)
-				{
-					FCbObjectView RootObjectView(GetBuffer.GetData());
-
-					FCbObjectView RequestsObjectView = RootObjectView["requests"].AsObjectView();
-					FZenRequestStats& RequestStats = Stats.RequestStats;
-
-					RequestStats.Count = RequestsObjectView["count"].AsInt64();
-					RequestStats.RateMean = RequestsObjectView["rate_mean"].AsDouble();
-					RequestStats.TAverage = RequestsObjectView["t_avg"].AsDouble();
-					RequestStats.TMin = RequestsObjectView["t_min"].AsDouble();
-					RequestStats.TMax = RequestsObjectView["t_max"].AsDouble();
-
-					FCbObjectView CacheObjectView = RootObjectView["cache"].AsObjectView();
-					FZenCacheStats& CacheStats = Stats.CacheStats;
-
-					CacheStats.Hits = CacheObjectView["hits"].AsInt64();
-					CacheStats.Misses = CacheObjectView["misses"].AsInt64();
-					CacheStats.HitRatio = CacheObjectView["hit_ratio"].AsDouble();
-					CacheStats.UpstreamHits = CacheObjectView["upstream_hits"].AsInt64();
-					CacheStats.UpstreamRatio = CacheObjectView["upstream_ratio"].AsDouble();
-
-					FCbObjectView CacheSizeObjectView = CacheObjectView["size"].AsObjectView();
-					FZenCacheSizeStats& CacheSizeStats = CacheStats.Size;
-					CacheSizeStats.Disk = CacheSizeObjectView["disk"].AsDouble();
-					CacheSizeStats.Memory = CacheSizeObjectView["memory"].AsDouble();
+		CacheStatsRequest = Async(ThreadPool, [this]
+		{
+			UE::Zen::FZenHttpRequest& Request = *CacheStatsHttpRequest.Get();
+			Request.Reset();
 
-					FCbObjectView UpstreamObjectView = RootObjectView["upstream"].AsObjectView();
-					FZenUpstreamStats& UpstreamStats = Stats.UpstreamStats;
+			TArray64<uint8> GetBuffer;
+			FZenHttpRequest::Result Result = Request.PerformBlockingDownload(TEXTVIEW("/stats/z$"), &GetBuffer, Zen::EContentType::CbObject);
 
-					UpstreamStats.Reading = UpstreamObjectView["reading"].AsBool();
-					UpstreamStats.Writing = UpstreamObjectView["writing"].AsBool();
-					UpstreamStats.WorkerThreads = UpstreamObjectView["worker_threads"].AsInt64();
-					UpstreamStats.QueueCount = UpstreamObjectView["queue_count"].AsInt64();
-					UpstreamStats.TotalUploadedMB = 0.0;
-					UpstreamStats.TotalDownloadedMB = 0.0;
+			FZenCacheStats Stats;
 
-					FCbObjectView UpstreamRequestObjectView = RootObjectView["upstream_gets"].AsObjectView();
-					FZenRequestStats& UpstreamRequestStats = Stats.UpstreamRequestStats;
+			if (Result == Zen::FZenHttpRequest::Result::Success && Request.GetResponseCode() == 200)
+			{
+				FCbFieldView RootView(GetBuffer.GetData());
+				Stats.bIsValid = LoadFromCompactBinary(RootView, Stats);
+			}
 
-					UpstreamRequestStats.Count = UpstreamRequestObjectView["count"].AsInt64();
-					UpstreamRequestStats.RateMean = UpstreamRequestObjectView["rate_mean"].AsDouble();
-					UpstreamRequestStats.TAverage = UpstreamRequestObjectView["t_avg"].AsDouble();
-					UpstreamRequestStats.TMin = UpstreamRequestObjectView["t_min"].AsDouble();
-					UpstreamRequestStats.TMax = UpstreamRequestObjectView["t_max"].AsDouble();
+			return Stats;
+		});
+	}
 
-					FCbArrayView EndpPointArrayView = UpstreamObjectView["endpoints"].AsArrayView();
+	return Stats.bIsValid;
+}
 
-					for (FCbFieldView FieldView : EndpPointArrayView)
-					{
-						FCbObjectView EndPointView = FieldView.AsObjectView();
-						FZenEndPointStats EndPointStats;
+bool 
+FZenServiceInstance::GetProjectStats(FZenProjectStats& Stats)
+{
+	check(IsInGameThread());
 
-						EndPointStats.Name = FString(EndPointView["name"].AsString());
-						EndPointStats.Url = FString(EndPointView["url"].AsString());
-						EndPointStats.Health = FString(EndPointView["state"].AsString());
+	// If we've already requested stats and they are ready then grab them
+	if ( ProjectStatsRequest.IsReady() == true )
+	{
+		LastProjectStats		= ProjectStatsRequest.Get();
+		LastProjectStatsTime	= FPlatformTime::Cycles64();
 
-						if (FCbObjectView Cache = EndPointView["cache"].AsObjectView())
-						{
-							EndPointStats.HitRatio = Cache["hit_ratio"].AsDouble();
-							EndPointStats.UploadedMB = Cache["put_bytes"].AsDouble() / 1024.0 / 1024.0;
-							EndPointStats.DownloadedMB = Cache["get_bytes"].AsDouble() / 1024.0 / 1024.0;
-							EndPointStats.ErrorCount = Cache["error_count"].AsInt64();
-						}
+		ProjectStatsRequest.Reset();
+	}
+	
+	// Make a copy of the last updated stats
+	Stats = LastProjectStats;
 
-						UpstreamStats.TotalUploadedMB += EndPointStats.UploadedMB;
-						UpstreamStats.TotalDownloadedMB += EndPointStats.DownloadedMB;
+	const uint64 CurrentTime = FPlatformTime::Cycles64();
+	constexpr double MinTimeBetweenRequestsInSeconds = 0.5;
+	const double DeltaTimeInSeconds = FPlatformTime::ToSeconds64(CurrentTime - LastProjectStatsTime);
 
-						UpstreamStats.EndPointStats.Push(EndPointStats);
-					}
+	if (!ProjectStatsRequest.IsValid() && DeltaTimeInSeconds > MinTimeBetweenRequestsInSeconds)
+	{
+#if WITH_EDITOR
+		EAsyncExecution ThreadPool = EAsyncExecution::LargeThreadPool;
+#else
+		EAsyncExecution ThreadPool = EAsyncExecution::ThreadPool;
+#endif
+		if (!ProjectStatsHttpRequest.IsValid())
+		{
+			TStringBuilder<128> ZenDomain;
+			ZenDomain << HostName << TEXT(":") << Port;
+			ProjectStatsHttpRequest = MakePimpl<FZenHttpRequest>(ZenDomain.ToString(), false);
+		}
 
-					FCbObjectView CASObjectView = RootObjectView["cid"].AsObjectView();
-					FCbObjectView CASSizeObjectView = CASObjectView["size"].AsObjectView();
+		// We've not got any requests in flight and we've met a given time requirement for requests
+		ProjectStatsRequest = Async(ThreadPool, [this]
+		{
+			UE::Zen::FZenHttpRequest& Request = *ProjectStatsHttpRequest.Get();
+			Request.Reset();
 
-					FZenCASSizeStats& CASSizeStats = Stats.CASStats.Size;
+			TArray64<uint8> GetBuffer;
+			FZenHttpRequest::Result Result = Request.PerformBlockingDownload(TEXTVIEW("/stats/prj"), &GetBuffer, Zen::EContentType::CbObject);
 
-					CASSizeStats.Tiny = CASSizeObjectView["tiny"].AsInt64();
-					CASSizeStats.Small = CASSizeObjectView["small"].AsInt64();
-					CASSizeStats.Large = CASSizeObjectView["large"].AsInt64();
-					CASSizeStats.Total = CASSizeObjectView["total"].AsInt64();
+			FZenProjectStats Stats;
 
-					Stats.IsValid = true;
-				}
+			if (Result == Zen::FZenHttpRequest::Result::Success && Request.GetResponseCode() == 200)
+			{
+				FCbFieldView RootView(GetBuffer.GetData());
+				Stats.bIsValid = LoadFromCompactBinary(RootView, Stats);
+			}
 
-				return Stats;
-			});
+			return Stats;
+		});
 	}
 
-	return Stats.IsValid;
+	return Stats.bIsValid;
 }
 
 bool 
@@ -2352,86 +2355,86 @@ FZenServiceInstance::RequestGC(const bool* OverrideCollectSmallObjects, const ui
 bool 
 FZenServiceInstance::GatherAnalytics(TArray<FAnalyticsEventAttribute>& Attributes)
 {
-	FZenStats ZenStats;
+	FZenCacheStats ZenStats;
 
-	if (GetStats(ZenStats) == false)
+	if (GetCacheStats(ZenStats) == false)
 		return false;
 
 	const FString BaseName = TEXT("Zen_");
 
 	{
 		FString AttrName = BaseName + TEXT("Enabled");
-		Attributes.Emplace(MoveTemp(AttrName), ZenStats.IsValid);
+		Attributes.Emplace(MoveTemp(AttrName), ZenStats.bIsValid);
 	}
 
 	{
 		FString AttrName = BaseName + TEXT("Cache_HitRatio");
-		Attributes.Emplace(MoveTemp(AttrName), ZenStats.CacheStats.HitRatio);
+		Attributes.Emplace(MoveTemp(AttrName), ZenStats.General.HitRatio);
 	}
 
 	{
 		FString AttrName = BaseName + TEXT("Cache_Hits");
-		Attributes.Emplace(MoveTemp(AttrName), ZenStats.CacheStats.Hits);
+		Attributes.Emplace(MoveTemp(AttrName), ZenStats.General.Hits);
 	}
 
 	{
 		FString AttrName = BaseName + TEXT("Cache_Misses");
-		Attributes.Emplace(MoveTemp(AttrName), ZenStats.CacheStats.Misses);
+		Attributes.Emplace(MoveTemp(AttrName), ZenStats.General.Misses);
 	}
 
 	{
 		FString AttrName = BaseName + TEXT("Cache_Size_Disk");
-		Attributes.Emplace(MoveTemp(AttrName), ZenStats.CacheStats.Size.Disk);
+		Attributes.Emplace(MoveTemp(AttrName), ZenStats.General.Size.Disk);
 	}
 
 	{
 		FString AttrName = BaseName + TEXT("Cache_Size_Memory");
-		Attributes.Emplace(MoveTemp(AttrName), ZenStats.CacheStats.Size.Memory);
+		Attributes.Emplace(MoveTemp(AttrName), ZenStats.General.Size.Memory);
 	}
 
 	{
 		FString AttrName = BaseName + TEXT("Cache_UpstreamHits");
-		Attributes.Emplace(MoveTemp(AttrName), ZenStats.CacheStats.UpstreamHits);
+		Attributes.Emplace(MoveTemp(AttrName), ZenStats.General.UpstreamHits);
 	}
 
 	{
 		FString AttrName = BaseName + TEXT("Cache_UpstreamRatio");
-		Attributes.Emplace(MoveTemp(AttrName), ZenStats.CacheStats.UpstreamRatio);
+		Attributes.Emplace(MoveTemp(AttrName), ZenStats.General.UpstreamRatio);
 	}
 
 	{
 		FString AttrName = BaseName + TEXT("Cache_TotalUploadedMB");
-		Attributes.Emplace(MoveTemp(AttrName), ZenStats.UpstreamStats.TotalUploadedMB);
+		Attributes.Emplace(MoveTemp(AttrName), ZenStats.Upstream.TotalUploadedMB);
 	}
 
 	{
 		FString AttrName = BaseName + TEXT("Upstream_TotalDownloadedMB");
-		Attributes.Emplace(MoveTemp(AttrName), ZenStats.UpstreamStats.TotalDownloadedMB);
+		Attributes.Emplace(MoveTemp(AttrName), ZenStats.Upstream.TotalDownloadedMB);
 	}
 
 	{
 		FString AttrName = BaseName + TEXT("Upstream_TotalUploadedMB");
-		Attributes.Emplace(MoveTemp(AttrName), ZenStats.UpstreamStats.TotalUploadedMB);
+		Attributes.Emplace(MoveTemp(AttrName), ZenStats.Upstream.TotalUploadedMB);
 	}
 
 	{
 		FString AttrName = BaseName + TEXT("Cas_Size_Large");
-		Attributes.Emplace(MoveTemp(AttrName), ZenStats.CASStats.Size.Large);
+		Attributes.Emplace(MoveTemp(AttrName), ZenStats.CID.Size.Large);
 	}
 
 	{
 		FString AttrName = BaseName + TEXT("Cas_Size_Small");
-		Attributes.Emplace(MoveTemp(AttrName), ZenStats.CASStats.Size.Small);
+		Attributes.Emplace(MoveTemp(AttrName), ZenStats.CID.Size.Small);
 	}
 
 	{
 		FString AttrName = BaseName + TEXT("Cas_Size_Tiny");
-		Attributes.Emplace(MoveTemp(AttrName), ZenStats.CASStats.Size.Tiny);
+		Attributes.Emplace(MoveTemp(AttrName), ZenStats.CID.Size.Tiny);
 	}
 
 	{
 		FString AttrName = BaseName + TEXT("Cas_Size_Total");
-		Attributes.Emplace(MoveTemp(AttrName), ZenStats.CASStats.Size.Total);
+		Attributes.Emplace(MoveTemp(AttrName), ZenStats.CID.Size.Total);
 	}
 
 	return true;
diff --git a/Engine/Source/Developer/Zen/Private/ZenStatistics.cpp b/Engine/Source/Developer/Zen/Private/ZenStatistics.cpp
new file mode 100644
--- /dev/null
+++ b/Engine/Source/Developer/Zen/Private/ZenStatistics.cpp
@@ -0,0 +1,154 @@
+// Copyright Epic Games, Inc. All Rights Reserved.
+
+#include "Experimental/ZenStatistics.h"
+#include "Serialization/CompactBinarySerialization.h"
+
+namespace UE::Zen
+{
+
+bool LoadFromCompactBinary(FCbFieldView Field, FZenSizeStats& OutValue)
+{
+	bool bOk = Field.IsObject();
+	bOk = LoadFromCompactBinary(Field["disk"], OutValue.Disk) & bOk;
+	bOk = LoadFromCompactBinary(Field["memory"], OutValue.Memory) & bOk;
+	return bOk;
+}
+
+bool LoadFromCompactBinary(FCbFieldView Field, FZenCIDSizeStats& OutValue)
+{
+	bool bOk = Field.IsObject();
+	bOk = LoadFromCompactBinary(Field["tiny"], OutValue.Tiny) & bOk;
+	bOk = LoadFromCompactBinary(Field["small"], OutValue.Small) & bOk;
+	bOk = LoadFromCompactBinary(Field["large"], OutValue.Large) & bOk;
+	bOk = LoadFromCompactBinary(Field["total"], OutValue.Total) & bOk;
+	return bOk;
+}
+
+bool LoadFromCompactBinary(FCbFieldView Field, FZenCIDStats& OutValue)
+{
+	bool bOk = Field.IsObject();
+	bOk = LoadFromCompactBinary(Field["size"], OutValue.Size) & bOk;
+	return bOk;
+}
+
+bool LoadFromCompactBinary(FCbFieldView Field, FZenCacheStats::FGeneralStats& OutValue)
+{
+	bool bOk = Field.IsObject();
+	bOk = LoadFromCompactBinary(Field["size"], OutValue.Size) & bOk;
+	bOk = LoadFromCompactBinary(Field["hits"], OutValue.Hits) & bOk;
+	bOk = LoadFromCompactBinary(Field["misses"], OutValue.Misses) & bOk;
+	bOk = LoadFromCompactBinary(Field["hit_ratio"], OutValue.HitRatio) & bOk;
+	bOk = LoadFromCompactBinary(Field["upstream_hits"], OutValue.UpstreamHits) & bOk;
+	bOk = LoadFromCompactBinary(Field["upstream_ratio"], OutValue.UpstreamRatio) & bOk;
+	return bOk;
+}
+
+bool LoadFromCompactBinary(FCbFieldView Field, FZenCacheStats::FRequestStats& OutValue)
+{
+	bool bOk = Field.IsObject();
+	bOk = LoadFromCompactBinary(Field["count"], OutValue.Count) & bOk;
+	bOk = LoadFromCompactBinary(Field["rate_mean"], OutValue.RateMean) & bOk;
+	bOk = LoadFromCompactBinary(Field["t_avg"], OutValue.TAverage) & bOk;
+	bOk = LoadFromCompactBinary(Field["t_min"], OutValue.TMin) & bOk;
+	bOk = LoadFromCompactBinary(Field["t_max"], OutValue.TMax) & bOk;
+	return bOk;
+}
+
+bool LoadFromCompactBinary(FCbFieldView Field, FZenCacheStats::FEndPointStats& OutValue)
+{
+	bool bOk = Field.IsObject();
+	bOk = LoadFromCompactBinary(Field["name"], OutValue.Name) & bOk;
+	bOk = LoadFromCompactBinary(Field["url"], OutValue.Url) & bOk;
+	bOk = LoadFromCompactBinary(Field["state"], OutValue.Health) & bOk;
+
+	FCbFieldView CacheSubfield = Field["cache"];
+	bOk = CacheSubfield.IsObject() & bOk;
+	if (CacheSubfield.IsObject())
+	{
+		bOk = LoadFromCompactBinary(CacheSubfield["hit_ratio"], OutValue.HitRatio) & bOk;
+		bOk = LoadFromCompactBinary(CacheSubfield["put_bytes"], OutValue.UploadedMB) & bOk;
+		OutValue.UploadedMB = OutValue.UploadedMB / 1024.0 / 1024.0;
+		bOk = LoadFromCompactBinary(CacheSubfield["get_bytes"], OutValue.DownloadedMB) & bOk;
+		OutValue.DownloadedMB = OutValue.DownloadedMB / 1024.0 / 1024.0;
+		bOk = LoadFromCompactBinary(CacheSubfield["error_count"], OutValue.ErrorCount) & bOk;
+	}
+	return bOk;
+}
+
+bool LoadFromCompactBinary(FCbFieldView Field, FZenCacheStats::FUpstreamStats& OutValue)
+{
+	bool bOk = Field.IsObject();
+	bOk = LoadFromCompactBinary(Field["reading"], OutValue.Reading) & bOk;
+	bOk = LoadFromCompactBinary(Field["writing"], OutValue.Writing) & bOk;
+	bOk = LoadFromCompactBinary(Field["worker_threads"], OutValue.WorkerThreads) & bOk;
+	bOk = LoadFromCompactBinary(Field["queue_count"], OutValue.QueueCount) & bOk;
+	OutValue.TotalUploadedMB = 0.0;
+	OutValue.TotalDownloadedMB = 0.0;
+
+	FCbFieldView EndpointsSubfield = Field["endpoints"];
+	if (!EndpointsSubfield.IsNull())
+	{
+		bOk = EndpointsSubfield.IsArray() & bOk;
+
+		for (FCbFieldView EndpointFieldView : EndpointsSubfield)
+		{
+			FZenCacheStats::FEndPointStats& NewEndPoint = OutValue.EndPoint.AddDefaulted_GetRef();
+			bOk = LoadFromCompactBinary(EndpointFieldView, NewEndPoint) & bOk;
+			OutValue.TotalUploadedMB += NewEndPoint.UploadedMB;
+			OutValue.TotalDownloadedMB += NewEndPoint.DownloadedMB;
+		}
+	}
+
+	return bOk;
+}
+
+bool LoadFromCompactBinary(FCbFieldView Field, FZenCacheStats& OutValue)
+{
+	bool bOk = Field.IsObject();
+	bOk = LoadFromCompactBinary(Field["cache"], OutValue.General) & bOk;
+	bOk = LoadFromCompactBinary(Field["requests"], OutValue.Request) & bOk;
+	bOk = LoadFromCompactBinary(Field["upstream"], OutValue.Upstream) & bOk;
+	bOk = LoadFromCompactBinary(Field["upstream_gets"], OutValue.UpstreamRequest) & bOk;
+	bOk = LoadFromCompactBinary(Field["cid"], OutValue.CID) & bOk;
+	return bOk;
+}
+
+bool LoadFromCompactBinary(FCbFieldView Field, FZenProjectStats::FReadWriteDeleteStats& OutValue)
+{
+	bool bOk = Field.IsObject();
+	bOk = LoadFromCompactBinary(Field["readcount"], OutValue.ReadCount) & bOk;
+	bOk = LoadFromCompactBinary(Field["writecount"], OutValue.WriteCount) & bOk;
+	bOk = LoadFromCompactBinary(Field["deletecount"], OutValue.DeleteCount) & bOk;
+	return bOk;
+}
+
+bool LoadFromCompactBinary(FCbFieldView Field, FZenProjectStats::FHitMissWriteStats& OutValue)
+{
+	bool bOk = Field.IsObject();
+	bOk = LoadFromCompactBinary(Field["hitcount"], OutValue.HitCount) & bOk;
+	bOk = LoadFromCompactBinary(Field["misscount"], OutValue.MissCount) & bOk;
+	bOk = LoadFromCompactBinary(Field["writecount"], OutValue.WriteCount) & bOk;
+	return bOk;
+}
+
+bool LoadFromCompactBinary(FCbFieldView Field, FZenProjectStats::FGeneralStats& OutValue)
+{
+	bool bOk = Field.IsObject();
+	bOk = LoadFromCompactBinary(Field["size"], OutValue.Size) & bOk;
+	bOk = LoadFromCompactBinary(Field["project"], OutValue.Project) & bOk;
+	bOk = LoadFromCompactBinary(Field["oplog"], OutValue.Oplog) & bOk;
+	bOk = LoadFromCompactBinary(Field["op"], OutValue.Op) & bOk;
+	bOk = LoadFromCompactBinary(Field["chunk"], OutValue.Chunk) & bOk;
+	bOk = LoadFromCompactBinary(Field["badrequestcount"], OutValue.BadRequestCount) & bOk;
+	return bOk;
+}
+
+bool LoadFromCompactBinary(FCbFieldView Field, FZenProjectStats& OutValue)
+{
+	bool bOk = Field.IsObject();
+	bOk = LoadFromCompactBinary(Field["store"], OutValue.General) & bOk;
+	bOk = LoadFromCompactBinary(Field["cid"], OutValue.CID) & bOk;
+	return bOk;
+}
+
+} // UE::Zen
diff --git a/Engine/Source/Developer/Zen/Public/Experimental/ZenServerInterface.h b/Engine/Source/Developer/Zen/Public/Experimental/ZenServerInterface.h
--- a/Engine/Source/Developer/Zen/Public/Experimental/ZenServerInterface.h
+++ b/Engine/Source/Developer/Zen/Public/Experimental/ZenServerInterface.h
@@ -187,7 +187,8 @@ public:
 
 	UE_API bool TryRecovery();
 
-	UE_API bool GetStats(FZenStats& Stats);
+	UE_API bool GetCacheStats(FZenCacheStats& Stats);
+	UE_API bool GetProjectStats(FZenProjectStats& Stats);
 	UE_API bool GatherAnalytics(TArray<FAnalyticsEventAttribute>& Attributes);
 
 	UE_API bool GetGCStatus(FGCStatus& Status);
@@ -208,10 +209,15 @@ private:
 	mutable TOptional<FGCStatus> LastGCStatus;
 	mutable uint64 LastGCStatusTime = 0;
 
-	mutable TPimplPtr<class FZenHttpRequest> StatsHttpRequest;
-	mutable TFuture<FZenStats> StatsRequest;
-	mutable FZenStats LastStats;
-	mutable uint64 LastStatsTime = 0;
+	mutable TPimplPtr<class FZenHttpRequest> CacheStatsHttpRequest;
+	mutable TFuture<FZenCacheStats> CacheStatsRequest;
+	mutable FZenCacheStats LastCacheStats;
+	mutable uint64 LastCacheStatsTime = 0;
+
+	mutable TPimplPtr<class FZenHttpRequest> ProjectStatsHttpRequest;
+	mutable TFuture<FZenProjectStats> ProjectStatsRequest;
+	mutable FZenProjectStats LastProjectStats;
+	mutable uint64 LastProjectStatsTime = 0;
 
 	FServiceSettings Settings;
 	FString URL;
diff --git a/Engine/Source/Developer/Zen/Public/Experimental/ZenStatistics.h b/Engine/Source/Developer/Zen/Public/Experimental/ZenStatistics.h
--- a/Engine/Source/Developer/Zen/Public/Experimental/ZenStatistics.h
+++ b/Engine/Source/Developer/Zen/Public/Experimental/ZenStatistics.h
@@ -5,27 +5,47 @@
 #include "Containers/UnrealString.h"
 #include "ZenGlobals.h"
 
+#define UE_API ZEN_API
+
 #if UE_WITH_ZEN
 
-namespace UE::Zen {
+class FCbFieldView;
 
-	struct FZenCacheSizeStats
-	{
-		double Disk = 0;
-		double Memory = 0;
-	};
+namespace UE::Zen
+{
+
+struct FZenSizeStats
+{
+	double Disk = 0;
+	double Memory = 0;
+};
 
-	struct FZenCacheStats
+struct FZenCIDSizeStats
+{
+	int64 Tiny = 0;
+	int64 Small = 0;
+	int64 Large = 0;
+	int64 Total = 0;
+};
+
+struct FZenCIDStats
+{
+	FZenCIDSizeStats Size;
+};
+
+struct FZenCacheStats
+{
+	struct FGeneralStats
 	{
-		FZenCacheSizeStats Size;
+		FZenSizeStats Size;
 		int64 Hits = 0;
 		int64 Misses = 0;
 		double HitRatio = 0.0;
 		int64 UpstreamHits = 0;
-		double UpstreamRatio = 0.0;	
+		double UpstreamRatio = 0.0;
 	};
 
-	struct FZenRequestStats
+	struct FRequestStats
 	{
 		int64 Count = 0;
 		double RateMean = 0.0;
@@ -34,7 +54,7 @@ namespace UE::Zen {
 		double TMax = 0.0;
 	};
 
-	struct FZenEndPointStats
+	struct FEndPointStats
 	{
 		FString Name;
 		FString Url;
@@ -45,7 +65,7 @@ namespace UE::Zen {
 		int64 ErrorCount = 0;
 	};
 
-	struct FZenUpstreamStats
+	struct FUpstreamStats
 	{
 		bool Reading = false;
 		bool Writing = false;
@@ -53,32 +73,62 @@ namespace UE::Zen {
 		int64 QueueCount = 0;
 		double TotalUploadedMB = 0;
 		double TotalDownloadedMB = 0;
-		TArray<FZenEndPointStats> EndPointStats;
+		TArray<FEndPointStats> EndPoint;
 	};
 
-	struct FZenCASSizeStats
+	FGeneralStats General;
+	FRequestStats Request;
+	FUpstreamStats Upstream;
+	FRequestStats UpstreamRequest;
+	FZenCIDStats CID;
+	bool bIsValid = false;
+};
+
+struct FZenProjectStats
+{
+	struct FReadWriteDeleteStats
 	{
-		int64 Tiny = 0;
-		int64 Small = 0;
-		int64 Large = 0;
-		int64 Total = 0;
+		int64 ReadCount;
+		int64 WriteCount;
+		int64 DeleteCount;
 	};
-
-	struct FZenCASStats
+	struct FHitMissWriteStats
 	{
-		FZenCASSizeStats Size;
+		int64 HitCount;
+		int64 MissCount;
+		int64 WriteCount;
 	};
-
-	struct FZenStats
+	struct FGeneralStats
 	{
-		FZenRequestStats RequestStats;
-		FZenCacheStats CacheStats;
-		FZenUpstreamStats UpstreamStats;
-		FZenRequestStats UpstreamRequestStats;
-		FZenCASStats CASStats;
-		bool IsValid = false;
+		FZenSizeStats Size;
+		FReadWriteDeleteStats Project;
+		FReadWriteDeleteStats Oplog;
+		FHitMissWriteStats Op;
+		FHitMissWriteStats Chunk;
+		int64 BadRequestCount;
 	};
+	FGeneralStats General;
+	FZenCIDStats CID;
+	bool bIsValid = false;
+};
+
+UE_API bool LoadFromCompactBinary(FCbFieldView Field, FZenSizeStats& OutValue);
+UE_API bool LoadFromCompactBinary(FCbFieldView Field, FZenCIDSizeStats& OutValue);
+UE_API bool LoadFromCompactBinary(FCbFieldView Field, FZenCIDStats& OutValue);
+
+UE_API bool LoadFromCompactBinary(FCbFieldView Field, FZenCacheStats::FGeneralStats& OutValue);
+UE_API bool LoadFromCompactBinary(FCbFieldView Field, FZenCacheStats::FRequestStats& OutValue);
+UE_API bool LoadFromCompactBinary(FCbFieldView Field, FZenCacheStats::FEndPointStats& OutValue);
+UE_API bool LoadFromCompactBinary(FCbFieldView Field, FZenCacheStats::FUpstreamStats& OutValue);
+UE_API bool LoadFromCompactBinary(FCbFieldView Field, FZenCacheStats& OutValue);
+
+UE_API bool LoadFromCompactBinary(FCbFieldView Field, FZenProjectStats::FReadWriteDeleteStats& OutValue);
+UE_API bool LoadFromCompactBinary(FCbFieldView Field, FZenProjectStats::FHitMissWriteStats& OutValue);
+UE_API bool LoadFromCompactBinary(FCbFieldView Field, FZenProjectStats::FGeneralStats& OutValue);
+UE_API bool LoadFromCompactBinary(FCbFieldView Field, FZenProjectStats& OutValue);
 
 } // namespace UE::Zen
 
 #endif // UE_WITH_ZEN
+
+#undef UE_API
diff --git a/Engine/Source/Programs/ZenDashboard/Private/SZenCacheStatistics.cpp b/Engine/Source/Programs/ZenDashboard/Private/SZenCacheStatistics.cpp
--- a/Engine/Source/Programs/ZenDashboard/Private/SZenCacheStatistics.cpp
+++ b/Engine/Source/Programs/ZenDashboard/Private/SZenCacheStatistics.cpp
@@ -49,11 +49,15 @@ TSharedRef<SWidget> SZenCacheStatistics::GetGridPanel()
 {
 	TSharedRef<SGridPanel> Panel = SNew(SGridPanel);
 
-	UE::Zen::FZenStats ZenStats;
+	bool bHaveStats = false;
+	UE::Zen::FZenCacheStats ZenStats;
 
 	if (TSharedPtr<UE::Zen::FZenServiceInstance> ServiceInstance = ZenServiceInstance.Get())
 	{
-		ServiceInstance->GetStats(ZenStats);
+		if (ServiceInstance->GetCacheStats(ZenStats))
+		{
+			bHaveStats = true;
+		}
 	}
 
 	const static FNumberFormattingOptions SingleDecimalFormatting = FNumberFormattingOptions()
@@ -71,6 +75,7 @@ TSharedRef<SWidget> SZenCacheStatistics::GetGridPanel()
 
 
 	Panel->AddSlot(0, Row)
+	.ColumnSpan(2)
 	[
 		SNew(STextBlock)
 		.Margin(FMargin(ColumnMargin, RowMargin))
@@ -79,51 +84,6 @@ TSharedRef<SWidget> SZenCacheStatistics::GetGridPanel()
 		.Text(LOCTEXT("Cache", "Local Cache"))
 	];
 
-	Panel->AddSlot(2, Row)
-	[
-		SNew(STextBlock)
-		.Margin(FMargin(ColumnMargin, RowMargin))
-		.ColorAndOpacity(TitleColor)
-		.Font(TitleFont)
-		.Text(LOCTEXT("CAS", "Local Content Store"))
-	];
-
-	Row++;
-
-	Panel->AddSlot(0, Row)
-	[
-		SNew(STextBlock)
-		.Margin(FMargin(ColumnMargin, RowMargin))
-		.Text(LOCTEXT("DiskSpace", "Disk Space"))
-	];
-
-	Panel->AddSlot(1, Row)
-	[
-		SNew(STextBlock)
-		.Margin(FMargin(ColumnMargin, RowMargin))
-		.Text_Lambda([DiskUsage = ZenStats.CacheStats.Size.Disk]
-		{
-			return FText::AsMemory(DiskUsage, (DiskUsage > 1024) ? &SingleDecimalFormatting : nullptr, nullptr, EMemoryUnitStandard::IEC);
-		})
-	];
-
-	Panel->AddSlot(2, Row)
-	[
-		SNew(STextBlock)
-		.Margin(FMargin(ColumnMargin, RowMargin))
-		.Text(LOCTEXT("DiskSpace", "Disk Space"))
-	];
-
-	Panel->AddSlot(3, Row)
-	[
-		SNew(STextBlock)
-		.Margin(FMargin(ColumnMargin, RowMargin))
-		.Text_Lambda([DiskUsage = ZenStats.CASStats.Size.Total]
-		{
-			return FText::AsMemory(DiskUsage, (DiskUsage > 1024) ? &SingleDecimalFormatting : nullptr, nullptr, EMemoryUnitStandard::IEC);
-		})
-	];
-
 	Row++;
 
 	Panel->AddSlot(0, Row)
@@ -137,7 +97,7 @@ TSharedRef<SWidget> SZenCacheStatistics::GetGridPanel()
 	[
 		SNew(STextBlock)
 		.Margin(FMargin(ColumnMargin, RowMargin))
-		.Text_Lambda([CacheStats = ZenStats.CacheStats]
+		.Text_Lambda([CacheStats = ZenStats.General]
 		{
 			int64 CacheTotal = CacheStats.Hits + CacheStats.Misses;
 			return CacheTotal == 0 ? LOCTEXT("CacheNoHitRateValue", "-") : FText::AsPercent(CacheStats.HitRatio, &SingleDecimalFormatting);
@@ -157,9 +117,9 @@ TSharedRef<SWidget> SZenCacheStatistics::GetGridPanel()
 	[
 		SNew(STextBlock)
 		.Margin(FMargin(ColumnMargin, RowMargin))
-		.Text_Lambda([CacheHits = ZenStats.CacheStats.Hits]
+		.Text_Lambda([bHaveStats, CacheHits = ZenStats.General.Hits]
 		{
-			return FText::AsNumber(CacheHits);
+			return bHaveStats ? FText::AsNumber(CacheHits) : LOCTEXT("UnavailableValue", "-");
 		})
 	];
 
@@ -169,22 +129,22 @@ TSharedRef<SWidget> SZenCacheStatistics::GetGridPanel()
 		[
 			SNew(STextBlock)
 			.Margin(FMargin(ColumnMargin, RowMargin))
-		.Text(LOCTEXT("CacheMissQuantity", "Misses"))
+			.Text(LOCTEXT("CacheMissQuantity", "Misses"))
 		];
 
 	Panel->AddSlot(1, Row)
 		[
 			SNew(STextBlock)
 			.Margin(FMargin(ColumnMargin, RowMargin))
-		.Text_Lambda([CacheMisses = ZenStats.CacheStats.Misses]
+			.Text_Lambda([bHaveStats, CacheMisses = ZenStats.General.Misses]
 			{
-				return FText::AsNumber(CacheMisses);
+				return bHaveStats ? FText::AsNumber(CacheMisses) : LOCTEXT("UnavailableValue", "-");
 			})
 		];
 
 	Row++;
 
-	if (!ZenStats.UpstreamStats.EndPointStats.IsEmpty())
+	if (!ZenStats.Upstream.EndPoint.IsEmpty())
 	{
 		Panel->AddSlot(0, Row)
 		[
@@ -232,7 +192,7 @@ TSharedRef<SWidget> SZenCacheStatistics::GetGridPanel()
 		];
 		Row++;
 
-		for (const UE::Zen::FZenEndPointStats& EndpointStats : ZenStats.UpstreamStats.EndPointStats)
+		for (const UE::Zen::FZenCacheStats::FEndPointStats& EndpointStats : ZenStats.Upstream.EndPoint)
 		{
 			Panel->AddSlot(0, Row)
 			[
diff --git a/Engine/Source/Programs/ZenDashboard/Private/SZenProjectStatistics.cpp b/Engine/Source/Programs/ZenDashboard/Private/SZenProjectStatistics.cpp
new file mode 100644
--- /dev/null
+++ b/Engine/Source/Programs/ZenDashboard/Private/SZenProjectStatistics.cpp
@@ -0,0 +1,213 @@
+// Copyright Epic Games, Inc. All Rights Reserved.
+
+#include "SZenProjectStatistics.h"
+#include "Experimental/ZenServerInterface.h"
+#include "Internationalization/FastDecimalFormat.h"
+#include "Math/BasicMathExpressionEvaluator.h"
+#include "Math/UnitConversion.h"
+#include "Misc/ExpressionParser.h"
+#include "Styling/StyleColors.h"
+#include "Widgets/Layout/SGridPanel.h"
+#include "Widgets/SBoxPanel.h"
+#include "Widgets/Text/STextBlock.h"
+
+#define LOCTEXT_NAMESPACE "ZenDashboard"
+
+void SZenProjectStatistics::Construct(const FArguments& InArgs)
+{
+	ZenServiceInstance = InArgs._ZenServiceInstance;
+
+	this->ChildSlot
+	[
+		SNew(SVerticalBox)
+
+		+ SVerticalBox::Slot()
+		.AutoHeight()
+		.Padding(0, 0, 0, 0)
+		.Expose(GridSlot)
+		[
+			GetGridPanel()
+		]
+	];
+
+	RegisterActiveTimer(0.5f, FWidgetActiveTimerDelegate::CreateSP(this, &SZenProjectStatistics::UpdateGridPanels));
+}
+
+EActiveTimerReturnType SZenProjectStatistics::UpdateGridPanels(double InCurrentTime, float InDeltaTime)
+{
+	(*GridSlot)
+	[
+		GetGridPanel()
+	];
+
+	SlatePrepass(GetPrepassLayoutScaleMultiplier());
+
+	return EActiveTimerReturnType::Continue;
+}
+
+TSharedRef<SWidget> SZenProjectStatistics::GetGridPanel()
+{
+	TSharedRef<SGridPanel> Panel = SNew(SGridPanel);
+
+	bool bHaveStats = false;
+	UE::Zen::FZenProjectStats ZenStats;
+
+	if (TSharedPtr<UE::Zen::FZenServiceInstance> ServiceInstance = ZenServiceInstance.Get())
+	{
+		if (ServiceInstance->GetProjectStats(ZenStats))
+		{
+			bHaveStats = true;
+		}
+	}
+
+	int64 TotalHits = ZenStats.General.Op.HitCount + ZenStats.General.Chunk.HitCount;
+	int64 TotalMisses = ZenStats.General.Op.MissCount + ZenStats.General.Chunk.MissCount;
+	int64 TotalWrites = ZenStats.General.Op.WriteCount + ZenStats.General.Chunk.WriteCount;
+
+	const static FNumberFormattingOptions SingleDecimalFormatting = FNumberFormattingOptions()
+		.SetUseGrouping(true)
+		.SetMinimumFractionalDigits(1)
+		.SetMaximumFractionalDigits(1);
+
+	int32 Row = 0;
+
+	const float RowMargin = 0.0f;
+	const float TitleMargin = 10.0f;
+	const float ColumnMargin = 10.0f;
+	const FSlateColor TitleColor = FStyleColors::AccentWhite;
+	const FSlateFontInfo TitleFont = FCoreStyle::GetDefaultFontStyle("Bold", 10);
+
+
+	Panel->AddSlot(0, Row)
+	.ColumnSpan(2)
+	[
+		SNew(STextBlock)
+		.Margin(FMargin(ColumnMargin, RowMargin))
+		.ColorAndOpacity(TitleColor)
+		.Font(TitleFont)
+		.Text(LOCTEXT("CAS", "Local Project Store"))
+	];
+
+	Row++;
+
+	Panel->AddSlot(0, Row)
+	[
+		SNew(STextBlock)
+		.Margin(FMargin(ColumnMargin, RowMargin))
+		.Text(LOCTEXT("ProjectHitRate", "Hit Rate"))
+	];
+
+	Panel->AddSlot(1, Row)
+	[
+		SNew(STextBlock)
+		.Margin(FMargin(ColumnMargin, RowMargin))
+		.Text_Lambda([bHaveStats, TotalHits, TotalMisses]
+		{
+			if (bHaveStats && ((TotalHits+TotalMisses) > 0))
+			{
+				double HitRate = static_cast<double>(TotalHits) / static_cast<double>(TotalHits + TotalMisses);
+				return FText::AsPercent(HitRate, &SingleDecimalFormatting);
+			}
+			return LOCTEXT("UnavailableValue", "-");
+		})
+	];
+
+	Row++;
+
+	Panel->AddSlot(0, Row)
+	[
+		SNew(STextBlock)
+		.Margin(FMargin(ColumnMargin, RowMargin))
+		.Text(LOCTEXT("ProjectHitQuantity", "Hits"))
+	];
+
+	Panel->AddSlot(1, Row)
+	[
+		SNew(STextBlock)
+		.Margin(FMargin(ColumnMargin, RowMargin))
+		.Text_Lambda([bHaveStats, TotalHits]
+		{
+			if (bHaveStats)
+			{
+				return FText::AsNumber(TotalHits);
+			}
+			return LOCTEXT("UnavailableValue", "-");
+		})
+	];
+
+	Row++;
+
+	Panel->AddSlot(0, Row)
+	[
+		SNew(STextBlock)
+		.Margin(FMargin(ColumnMargin, RowMargin))
+		.Text(LOCTEXT("ProjectMissQuantity", "Misses"))
+	];
+
+	Panel->AddSlot(1, Row)
+	[
+		SNew(STextBlock)
+		.Margin(FMargin(ColumnMargin, RowMargin))
+		.Text_Lambda([bHaveStats, TotalMisses]
+		{
+			if (bHaveStats)
+			{
+				return FText::AsNumber(TotalMisses);
+			}
+			return LOCTEXT("UnavailableValue", "-");
+		})
+	];
+
+	Row++;
+
+	Panel->AddSlot(0, Row)
+	[
+		SNew(STextBlock)
+		.Margin(FMargin(ColumnMargin, RowMargin))
+		.Text(LOCTEXT("ProjectWriteQuantity", "Writes"))
+	];
+
+	Panel->AddSlot(1, Row)
+	[
+		SNew(STextBlock)
+		.Margin(FMargin(ColumnMargin, RowMargin))
+		.Text_Lambda([bHaveStats, TotalWrites]
+		{
+			if (bHaveStats)
+			{
+				return FText::AsNumber(TotalWrites);
+			}
+			return LOCTEXT("UnavailableValue", "-");
+		})
+	];
+
+	Row++;
+
+	Panel->AddSlot(0, Row)
+	[
+		SNew(STextBlock)
+		.Margin(FMargin(ColumnMargin, RowMargin))
+		.Text(LOCTEXT("ProjectBadRequests", "Bad Requests"))
+	];
+
+	Panel->AddSlot(1, Row)
+	[
+		SNew(STextBlock)
+		.Margin(FMargin(ColumnMargin, RowMargin))
+		.Text_Lambda([bHaveStats, BadRequests = ZenStats.General.BadRequestCount]
+		{
+			if (bHaveStats)
+			{
+				return FText::AsNumber(BadRequests);
+			}
+			return LOCTEXT("UnavailableValue", "-");
+		})
+	];
+
+	Row++;
+
+
+	return Panel;
+}
+
+#undef LOCTEXT_NAMESPACE
diff --git a/Engine/Source/Programs/ZenDashboard/Private/SZenProjectStatistics.h b/Engine/Source/Programs/ZenDashboard/Private/SZenProjectStatistics.h
new file mode 100644
--- /dev/null
+++ b/Engine/Source/Programs/ZenDashboard/Private/SZenProjectStatistics.h
@@ -0,0 +1,33 @@
+// Copyright Epic Games, Inc. All Rights Reserved.
+
+#pragma once
+
+#include "Experimental/ZenServerInterface.h"
+#include "ServiceInstanceManager.h"
+#include "Templates/SharedPointer.h"
+#include "Widgets/DeclarativeSyntaxSupport.h"
+#include "Widgets/SBoxPanel.h"
+#include "Widgets/SCompoundWidget.h"
+#include "Widgets/SWidget.h"
+
+class SZenProjectStatistics : public SCompoundWidget
+{
+	SLATE_BEGIN_ARGS(SZenProjectStatistics)
+		: _ZenServiceInstance(nullptr)
+	{ }
+
+	SLATE_ATTRIBUTE(TSharedPtr<UE::Zen::FZenServiceInstance>, ZenServiceInstance);
+
+	SLATE_END_ARGS()
+
+	void Construct(const FArguments& InArgs);
+
+private:
+
+	TSharedRef<SWidget> GetGridPanel();
+
+	EActiveTimerReturnType UpdateGridPanels(double InCurrentTime, float InDeltaTime);
+
+	SVerticalBox::FSlot* GridSlot = nullptr;
+	TAttribute<TSharedPtr<UE::Zen::FZenServiceInstance>> ZenServiceInstance;
+};
\ No newline at end of file
diff --git a/Engine/Source/Programs/ZenDashboard/Private/SZenServiceStatus.cpp b/Engine/Source/Programs/ZenDashboard/Private/SZenServiceStatus.cpp
--- a/Engine/Source/Programs/ZenDashboard/Private/SZenServiceStatus.cpp
+++ b/Engine/Source/Programs/ZenDashboard/Private/SZenServiceStatus.cpp
@@ -48,11 +48,17 @@ EActiveTimerReturnType SZenServiceStatus::UpdateState(double InCurrentTime, floa
 		NextState.bIsRunning = UE::Zen::IsLocalServiceRunning(*NextState.RunContext.GetDataPath(), &NextState.LocalPort);
 	}
 
+	NextState.GCStatus = UE::Zen::FGCStatus{};
+	NextState.bHaveStats = false;
 	if (NextState.bIsRunning)
 	{
 		if (TSharedPtr<UE::Zen::FZenServiceInstance> ServiceInstance = ZenServiceInstance.Get())
 		{
 			ServiceInstance->GetGCStatus(NextState.GCStatus);
+			if (ServiceInstance->GetCacheStats(NextState.ZenCacheStats) && ServiceInstance->GetProjectStats(NextState.ZenProjectStats))
+			{
+				NextState.bHaveStats = true;
+			}
 		}
 	}
 
@@ -81,6 +87,11 @@ const SZenServiceStatus::FState& SZenServiceStatus::GetCurrentState() const
 TSharedRef<SWidget> SZenServiceStatus::GetGridPanel()
 {
 	TSharedRef<SGridPanel> Panel = SNew(SGridPanel);
+	
+	const static FNumberFormattingOptions SingleDecimalFormatting = FNumberFormattingOptions()
+		.SetUseGrouping(true)
+		.SetMinimumFractionalDigits(1)
+		.SetMaximumFractionalDigits(1);
 
 	int32 Row = 0;
 
@@ -189,6 +200,29 @@ TSharedRef<SWidget> SZenServiceStatus::GetGridPanel()
 
 	++Row;
 
+	Panel->AddSlot(0, Row)
+	[
+		SNew(STextBlock)
+		.Margin(FMargin(ColumnMargin, RowMargin))
+		.ColorAndOpacity(TitleColor)
+		.Font(TitleFont)
+		.Text(LOCTEXT("ServiceStatus_DiskSpace", "Disk Space"))
+	];
+
+	Panel->AddSlot(1, Row)
+	[
+		SNew(STextBlock)
+		.Margin(FMargin(ColumnMargin, RowMargin))
+		.Text_Lambda([this]
+		{
+			const FState& CurrentState = GetCurrentState();
+			int64 TotalDiskSpace = CurrentState.ZenCacheStats.General.Size.Disk + CurrentState.ZenProjectStats.General.Size.Disk + CurrentState.ZenCacheStats.CID.Size.Total;
+			return CurrentState.bHaveStats ? FText::AsMemory(TotalDiskSpace, (TotalDiskSpace > 1024) ? &SingleDecimalFormatting : nullptr, nullptr, EMemoryUnitStandard::IEC) : LOCTEXT("UnavailableValue", "-");
+		})
+	];
+
+	++Row;
+
 	Panel->AddSlot(0, Row)
 	.VAlign(VAlign_Bottom)
 	[
diff --git a/Engine/Source/Programs/ZenDashboard/Private/SZenServiceStatus.h b/Engine/Source/Programs/ZenDashboard/Private/SZenServiceStatus.h
--- a/Engine/Source/Programs/ZenDashboard/Private/SZenServiceStatus.h
+++ b/Engine/Source/Programs/ZenDashboard/Private/SZenServiceStatus.h
@@ -28,7 +28,10 @@ private:
 		UE::Zen::FZenLocalServiceRunContext RunContext;
 		FString Version;
 		UE::Zen::FGCStatus GCStatus;
+		UE::Zen::FZenCacheStats ZenCacheStats;
+		UE::Zen::FZenProjectStats ZenProjectStats;
 		uint16 LocalPort = 0;
+		bool bHaveStats = false;
 		bool bGotRunContext = false;
 		bool bIsRunning = false;
 	};
diff --git a/Engine/Source/Programs/ZenDashboard/Private/ZenDashboard.cpp b/Engine/Source/Programs/ZenDashboard/Private/ZenDashboard.cpp
--- a/Engine/Source/Programs/ZenDashboard/Private/ZenDashboard.cpp
+++ b/Engine/Source/Programs/ZenDashboard/Private/ZenDashboard.cpp
@@ -19,6 +19,7 @@
 #include "SMessageDialog.h"
 #include "StandaloneRenderer.h"
 #include "SZenCacheStatistics.h"
+#include "SZenProjectStatistics.h"
 #include "SZenServiceStatus.h"
 #include "Templates/SharedPointer.h"
 #include "Widgets/Input/SButton.h"
@@ -469,8 +470,25 @@ public:
 					.Padding(0.0f, 10.0f, 0.0f, 0.0f)
 					.VAlign(VAlign_Top)
 					[
-						SNew(SZenCacheStatistics)
-						.ZenServiceInstance(ServiceInstanceManager.ToSharedRef(), &UE::Zen::FServiceInstanceManager::GetZenServiceInstance)
+						SNew(SHorizontalBox)
+
+						+SHorizontalBox::Slot()
+						.AutoWidth()
+						.HAlign(HAlign_Left)
+						.VAlign(VAlign_Top)
+						[
+							SNew(SZenCacheStatistics)
+							.ZenServiceInstance(ServiceInstanceManager.ToSharedRef(), &UE::Zen::FServiceInstanceManager::GetZenServiceInstance)
+						]
+
+						+SHorizontalBox::Slot()
+						.AutoWidth()
+						.HAlign(HAlign_Left)
+						.VAlign(VAlign_Top)
+						[
+							SNew(SZenProjectStatistics)
+							.ZenServiceInstance(ServiceInstanceManager.ToSharedRef(), &UE::Zen::FServiceInstanceManager::GetZenServiceInstance)
+						]
 					]
 				]
 			];
