Avoid attempting to pop a messagebox from UE Zen interface when we're in additional unattended scenarios like commandlets and programs.

diff --git a/Engine/Source/Developer/Zen/Private/ZenServerInterface.cpp b/Engine/Source/Developer/Zen/Private/ZenServerInterface.cpp
--- a/Engine/Source/Developer/Zen/Private/ZenServerInterface.cpp
+++ b/Engine/Source/Developer/Zen/Private/ZenServerInterface.cpp
@@ -1804,7 +1804,7 @@ FZenServiceInstance::Initialize()
 static void
 PromptUserToStopRunningServerInstanceForUpdate(const FString& ServerFilePath)
 {
-	if (FApp::IsUnattended())
+	if (IS_PROGRAM || FApp::IsUnattended() || IsRunningCommandlet() || GIsRunningUnattendedScript)
 	{
 		// Just log as there is no one to show a message
 		UE_LOG(LogZenServiceInstance, Display, TEXT("ZenServer needs to be updated to a new version. Please shut down any tools that are using the ZenServer at '%s'"), *ServerFilePath);
@@ -1819,7 +1819,7 @@ PromptUserToStopRunningServerInstanceForUpdate(const FString& ServerFilePath)
 static void
 PromptUserOfLockedDataFolder(const FString& DataPath)
 {
-	if (FApp::IsUnattended())
+	if (IS_PROGRAM || FApp::IsUnattended() || IsRunningCommandlet() || GIsRunningUnattendedScript)
 	{
 		// Just log as there is no one to show a message
 		UE_LOG(LogZenServiceInstance, Warning, TEXT("ZenServer Failed to auto launch, an unknown process is locking the data folder '%s'"), *DataPath);
@@ -1834,7 +1834,7 @@ PromptUserOfLockedDataFolder(const FString& DataPath)
 static void
 PromptUserOfFailedShutDownOfExistingProcess(uint16 Port)
 {
-	if (FApp::IsUnattended())
+	if (IS_PROGRAM || FApp::IsUnattended() || IsRunningCommandlet() || GIsRunningUnattendedScript)
 	{
 		// Just log as there is no one to show a message
 		UE_LOG(LogZenServiceInstance, Warning, TEXT("ZenServer Failed to auto launch, failed to shut down currently running service using port %u"), Port);
