Fix a bug with handling of Zen interprocess communication specific to Mac/Linux.  When signaling a semaphore you haven't created, don't remove its ID, leave that for the creator to do.

diff --git a/Engine/Source/Developer/Zen/Private/ZenServerInterface.cpp b/Engine/Source/Developer/Zen/Private/ZenServerInterface.cpp
--- a/Engine/Source/Developer/Zen/Private/ZenServerInterface.cpp
+++ b/Engine/Source/Developer/Zen/Private/ZenServerInterface.cpp
@@ -1126,7 +1126,6 @@ RequestZenShutdownOnEffectivePort(uint16 EffectiveListenPort)
 	}
 
 	semctl(Semaphore, 0, SETVAL, 0);
-	semctl(Semaphore, 0, IPC_RMID);
 	return true;
 #else
 	static_assert(false, "Missing implementation for Zen named shutdown events");
