On Windows, try to detect if the zenserver data path is configured to use a google drive and warn the user that it is not recommended.
Increase log level to warning if invalid or bad data path is configured

diff --git a/Engine/Source/Developer/Zen/Private/ZenServerInterface.cpp b/Engine/Source/Developer/Zen/Private/ZenServerInterface.cpp
--- a/Engine/Source/Developer/Zen/Private/ZenServerInterface.cpp
+++ b/Engine/Source/Developer/Zen/Private/ZenServerInterface.cpp
@@ -654,7 +654,7 @@ PromptUserUnableToDetermineValidDataPath()
 	if (FApp::IsUnattended())
 	{
 		// Just log as there is no one to show a message
-		UE_LOG(LogZenServiceInstance, Display, TEXT("ZenServer is unable to determine a valid data path"));
+		UE_LOG(LogZenServiceInstance, Warning, TEXT("ZenServer is unable to determine a valid data path"));
 		return;
 	}
 	FString LogDirPath = FPaths::ConvertRelativePathToFull(FPaths::ProjectLogDir());
@@ -669,7 +669,7 @@ PromptUserAboutInvalidValidDataPathConfiguration(const FString& UsedDataPath)
 	if (FApp::IsUnattended())
 	{
 		// Just log as there is no one to show a message
-		UE_LOG(LogZenServiceInstance, Display, TEXT("ZenServer has detected invalid data path configuration. Falling back to '%s'"), *UsedDataPath);
+		UE_LOG(LogZenServiceInstance, Warning, TEXT("ZenServer has detected invalid data path configuration. Falling back to '%s'"), *UsedDataPath);
 		return;
 	}
 
@@ -679,6 +679,23 @@ PromptUserAboutInvalidValidDataPathConfiguration(const FString& UsedDataPath)
 	FPlatformMisc::MessageBoxExt(EAppMsgType::Ok, *ZenInvalidValidDataPathConfigurationPromptText.ToString(), *ZenInvalidValidDataPathConfigurationPromptTitle.ToString());
 }
 
+#if PLATFORM_WINDOWS
+static void
+PromptUserIsUsingGoogleDriveAsDataPath()
+{
+	if (FApp::IsUnattended())
+	{
+		// Just log as there is no one to show a message
+		UE_LOG(LogZenServiceInstance, Warning, TEXT("ZenServer is configured to use Google Drive as a data path, this is highly inadvisable. Please use a path on a local physical drive."));
+		return;
+	}
+	FString LogDirPath = FPaths::ConvertRelativePathToFull(FPaths::ProjectLogDir());
+	FText ZenInvalidDataPathPromptTitle = NSLOCTEXT("Zen", "Zen_GoogleDriveDataPathPromptTitle", "Using Google Drive as a data path");
+	FText ZenInvalidDataPathPromptText = FText::Format(NSLOCTEXT("Zen", "Zen_GoogleDriveDataPathPromptText", "ZenServer is configured to use Google Drive as a data path, this is highly inadvisable.\nPlease use a data path on a local physical drive.\nCheck the log in '{0}' for details.\nUpdate your configuration and restart."), FText::FromString(LogDirPath));
+	FPlatformMisc::MessageBoxExt(EAppMsgType::Ok, *ZenInvalidDataPathPromptText.ToString(), *ZenInvalidDataPathPromptTitle.ToString());
+}
+#endif // PLATFORM_WINDOWS
+
 static void ReadCbField(FCbFieldView Field, UE::Zen::FZenSizeStats& OutValue)
 {
 	FCbObjectView ObjectView = Field.AsObjectView();
@@ -729,6 +746,37 @@ FServiceSettings::ReadFromConfig()
 			{
 				PromptUserAboutInvalidValidDataPathConfiguration(AutoLaunchSettings.DataPath);
 			}
+
+#if PLATFORM_WINDOWS
+			{
+				int32 DriveEnd = 0;
+				if (AutoLaunchSettings.DataPath.FindChar(':', DriveEnd))
+				{
+					FString DrivePath = AutoLaunchSettings.DataPath.Left(DriveEnd + 1);
+					TCHAR VolumeName[128];
+
+					BOOL OK = GetVolumeInformation(
+						*DrivePath,
+						VolumeName,
+						127,
+						NULL,
+						NULL,
+						NULL,
+						NULL,
+						NULL);
+
+					if (OK)
+					{
+						VolumeName[127] = 0;
+						if (FString(VolumeName) == TEXT("Google Drive"))
+						{
+							PromptUserIsUsingGoogleDriveAsDataPath();
+						}
+					}
+				}
+			}
+#endif // PLATFORM_WINDOWS
+
 			GConfig->GetString(AutoLaunchConfigSection, TEXT("ExtraArgs"), AutoLaunchSettings.ExtraArgs, GEngineIni);
 
 			ReadUInt16FromConfig(AutoLaunchConfigSection, TEXT("DesiredPort"), AutoLaunchSettings.DesiredPort, GEngineIni);
