Add logic to attempt shutdown of zenserver process even if we can't get to the process handle
Add attempt to hard terminate the zenserver process if we can't find the matching zenserver state

diff --git a/Engine/Source/Developer/Zen/Private/ZenServerInterface.cpp b/Engine/Source/Developer/Zen/Private/ZenServerInterface.cpp
--- a/Engine/Source/Developer/Zen/Private/ZenServerInterface.cpp
+++ b/Engine/Source/Developer/Zen/Private/ZenServerInterface.cpp
@@ -1181,6 +1181,54 @@ RequestZenShutdownOnEffectivePort(uint16 EffectiveListenPort)
 #endif
 }
 
+static bool AttemptShutdownUsingExecutablePathOnly(const ZenServerState& ServerState, const TCHAR* ExecutablePath, double MaximumWaitDurationSeconds)
+{
+	uint32 ServicePid;
+	if (!IsZenProcessActive(ExecutablePath, &ServicePid))
+	{
+		return true;
+	}
+
+	const ZenServerState::ZenServerEntry* Entry = ServerState.LookupByPid(ServicePid);
+	if (!Entry)
+	{
+		if (!IsZenProcessActive(ExecutablePath, nullptr))
+		{
+			return true;
+		}
+		UE_LOG(LogZenServiceInstance, Warning, TEXT("Can't find server state for running service for executable '%s' (Pid: %u)"), ExecutablePath, ServicePid);
+		return false;
+	}
+
+	uint16 EffectiveListenPort = Entry->EffectiveListenPort;
+	if (!RequestZenShutdownOnEffectivePort(EffectiveListenPort))
+	{
+		if (!IsZenProcessActive(ExecutablePath, &ServicePid))
+		{
+			return true;
+		}
+		UE_LOG(LogZenServiceInstance, Warning, TEXT("Failed to request shutdown for running service for executable '%s' (Pid: %u, Port: %u)"), ExecutablePath, ServicePid, EffectiveListenPort);
+		return false;
+	}
+
+	uint64 ZenShutdownWaitStartTime = FPlatformTime::Cycles64();
+	while ((Entry != nullptr) || IsZenProcessActive(ExecutablePath, &ServicePid))
+	{
+		double ZenShutdownWaitDuration = FPlatformTime::ToSeconds64(FPlatformTime::Cycles64() - ZenShutdownWaitStartTime);
+		if (ZenShutdownWaitDuration < MaximumWaitDurationSeconds)
+		{
+			FPlatformProcess::Sleep(0.01f);
+		}
+		else
+		{
+			UE_LOG(LogZenServiceInstance, Warning, TEXT("Timed out waiting for shutdown of running service for executable '%s' (Pid: %u)"), ExecutablePath, ServicePid);
+			return false;
+		}
+		Entry = ServerState.LookupByPid(ServicePid);
+	}
+	return true;
+}
+
 static bool
 ShutdownRunningServiceUsingExecutablePath(const TCHAR* ExecutablePath, double MaximumWaitDurationSeconds = 15.0)
 {
@@ -1190,36 +1238,71 @@ ShutdownRunningServiceUsingExecutablePath(const TCHAR* ExecutablePath, double Ma
 		UE_LOG(LogZenServiceInstance, Log, TEXT("No running service found for executable '%s'"), ExecutablePath);
 		return true;
 	}
+
 	UE_LOG(LogZenServiceInstance, Display, TEXT("Waiting for running instance using executable '%s' to shut down"), ExecutablePath);
+
 	FProcHandle ProcessHandle = FPlatformProcess::OpenProcess(ServicePid);
-	if (!ProcessHandle.IsValid() && IsZenProcessActive(ExecutablePath, nullptr))
-	{
-		UE_LOG(LogZenServiceInstance, Warning, TEXT("Failed to open handle for running service for executable '%s' (Pid: %u)"), ExecutablePath, ServicePid);
-		return false;
-	}
 	ON_SCOPE_EXIT{ FPlatformProcess::CloseProc(ProcessHandle); };
 
 	const ZenServerState ServerState;
 	const ZenServerState::ZenServerEntry* Entry = ServerState.LookupByPid(ServicePid);
 	if (!Entry)
 	{
-		if (FPlatformProcess::IsProcRunning(ProcessHandle))
+		if (!IsZenProcessActive(ExecutablePath, nullptr))
 		{
-			UE_LOG(LogZenServiceInstance, Warning, TEXT("Can't find server state for running service for executable '%s' (Pid: %u)"), ExecutablePath, ServicePid);
+			return true;
+		}
+		if (ProcessHandle.IsValid())
+		{
+			UE_LOG(LogZenServiceInstance, Warning, TEXT("Can't find server state for running service for executable '%s' (Pid: %u), attempting hard terminate"), ExecutablePath, ServicePid);
+			FPlatformProcess::TerminateProc(ProcessHandle);
+			FPlatformProcess::CloseProc(ProcessHandle);
+
+			if (!IsZenProcessActive(ExecutablePath, nullptr))
+			{
+				return true;
+			}
+
+			UE_LOG(LogZenServiceInstance, Warning, TEXT("Failed to terminate running service for executable '%s' (Pid: %u)"), ExecutablePath, ServicePid);
 			return false;
 		}
-		return true;
+		UE_LOG(LogZenServiceInstance, Warning, TEXT("Can't find server state for running service for executable '%s' (Pid: %u)"), ExecutablePath, ServicePid);
+		return false;
 	}
+
 	uint16 EffectiveListenPort = Entry->EffectiveListenPort;
+
+	if (!ProcessHandle.IsValid())
+	{
+		if (!IsZenProcessActive(ExecutablePath, nullptr))
+		{
+			return true;
+		}
+
+		UE_LOG(LogZenServiceInstance, Warning, TEXT("Failed to open handle for running service for executable '%s' (Pid: %u, Port: %u), attempting blind shut down"), ExecutablePath, ServicePid, EffectiveListenPort);
+		return AttemptShutdownUsingExecutablePathOnly(ServerState, ExecutablePath, MaximumWaitDurationSeconds);
+	}
+
 	if (!RequestZenShutdownOnEffectivePort(EffectiveListenPort))
 	{
-		if (FPlatformProcess::IsProcRunning(ProcessHandle) || IsZenProcessActive(ExecutablePath, nullptr))
+		if (!FPlatformProcess::IsProcRunning(ProcessHandle) && !IsZenProcessActive(ExecutablePath, nullptr))
 		{
-			UE_LOG(LogZenServiceInstance, Warning, TEXT("Failed to request shutdown for running service for executable '%s' (Pid: %u, Port: %u)"), ExecutablePath, ServicePid, EffectiveListenPort);
-			return false;
+			return true;
 		}
-		return true;
+
+		UE_LOG(LogZenServiceInstance, Warning, TEXT("Failed to request shutdown for running service for executable '%s' (Pid: %u, Port: %u), attempting hard terminate"), ExecutablePath, ServicePid, EffectiveListenPort);
+		FPlatformProcess::TerminateProc(ProcessHandle);
+		FPlatformProcess::CloseProc(ProcessHandle);
+
+		if (!IsZenProcessActive(ExecutablePath, nullptr))
+		{
+			return true;
+		}
+
+		UE_LOG(LogZenServiceInstance, Warning, TEXT("Failed to terminate running service for executable '%s' (Pid: %u, Port: %u)"), ExecutablePath, ServicePid, EffectiveListenPort);
+		return false;
 	}
+
 	uint64 ZenShutdownWaitStartTime = FPlatformTime::Cycles64();
 	while (FPlatformProcess::IsProcRunning(ProcessHandle) || IsZenProcessActive(ExecutablePath, nullptr))
 	{
@@ -1237,6 +1320,39 @@ ShutdownRunningServiceUsingExecutablePath(const TCHAR* ExecutablePath, double Ma
 	return true;
 }
 
+static bool AttemptShutdownUsingPortOnly(const ZenServerState& ServerState, uint16 EffectiveListenPort, double MaximumWaitDurationSeconds)
+{
+	if (!RequestZenShutdownOnEffectivePort(EffectiveListenPort))
+	{
+		if (!IsZenProcessUsingEffectivePort(EffectiveListenPort))
+		{
+			return true;
+		}
+		UE_LOG(LogZenServiceInstance, Warning, TEXT("Failed to request shutdown for running service using port %u"), EffectiveListenPort);
+		return false;
+	}
+
+	const ZenServerState::ZenServerEntry* Entry = ServerState.LookupByEffectiveListenPort(EffectiveListenPort);
+	uint64 ZenShutdownWaitStartTime = FPlatformTime::Cycles64();
+	while ((Entry != nullptr) || IsZenProcessUsingEffectivePort(EffectiveListenPort))
+	{
+		double ZenShutdownWaitDuration = FPlatformTime::ToSeconds64(FPlatformTime::Cycles64() - ZenShutdownWaitStartTime);
+		if (ZenShutdownWaitDuration < MaximumWaitDurationSeconds)
+		{
+			FPlatformProcess::Sleep(0.01f);
+		}
+		else
+		{
+			UE_LOG(LogZenServiceInstance, Warning, TEXT("Timed out waiting for shutdown of running service using port %u"), EffectiveListenPort);
+			return false;
+		}
+		Entry = ServerState.LookupByEffectiveListenPort(EffectiveListenPort);
+	}
+	// Sleep one second to ensure the process can run the final exit code and not hindering overwriting the executable
+	FPlatformProcess::Sleep(1.0f);
+	return true;
+}
+
 static bool
 ShutdownRunningServiceUsingEffectivePort(uint16 EffectiveListenPort, double MaximumWaitDurationSeconds = 15.0)
 {
@@ -1244,41 +1360,46 @@ ShutdownRunningServiceUsingEffectivePort(uint16 EffectiveListenPort, double Maxi
 	{
 		return true;
 	}
+
 	UE_LOG(LogZenServiceInstance, Display, TEXT("Waiting for running instance using port %u to shut down"), EffectiveListenPort);
+
 	const ZenServerState ServerState;
 	const ZenServerState::ZenServerEntry* Entry = ServerState.LookupByEffectiveListenPort(EffectiveListenPort);
 	if (!Entry)
 	{
-		if (IsZenProcessUsingEffectivePort(EffectiveListenPort))
+		if (!IsZenProcessUsingEffectivePort(EffectiveListenPort))
 		{
-			UE_LOG(LogZenServiceInstance, Warning, TEXT("Can't find server state for running service using port %u"), EffectiveListenPort);
-			return false;
+			return true;
 		}
-		return true;
+
+		UE_LOG(LogZenServiceInstance, Warning, TEXT("Can't find server state for running service using port %u, attempting blind shut down"), EffectiveListenPort);
+		return AttemptShutdownUsingPortOnly(ServerState, EffectiveListenPort, MaximumWaitDurationSeconds);
 	}
 
 	uint32 ServicePid = Entry->Pid;
 	FProcHandle ProcessHandle = FPlatformProcess::OpenProcess(ServicePid);
+	ON_SCOPE_EXIT{ FPlatformProcess::CloseProc(ProcessHandle); };
+
 	if (!ProcessHandle.IsValid())
 	{
-		if (IsZenProcessUsingEffectivePort(EffectiveListenPort))
+		if (!IsZenProcessUsingEffectivePort(EffectiveListenPort))
 		{
-			UE_LOG(LogZenServiceInstance, Warning, TEXT("Failed to open handle for running service using port %u (Pid: %u)"), EffectiveListenPort, ServicePid);
-			return false;
+			return true;
 		}
-		return true;
+		UE_LOG(LogZenServiceInstance, Warning, TEXT("Failed to open handle for running service using port %u (Pid: %u), attempting blind shut down"), EffectiveListenPort, ServicePid);
+		return AttemptShutdownUsingPortOnly(ServerState, EffectiveListenPort, MaximumWaitDurationSeconds);
 	}
-	ON_SCOPE_EXIT{ FPlatformProcess::CloseProc(ProcessHandle); };
 
 	if (!RequestZenShutdownOnEffectivePort(EffectiveListenPort))
 	{
-		if (FPlatformProcess::IsProcRunning(ProcessHandle) || IsZenProcessUsingEffectivePort(EffectiveListenPort))
+		if (!FPlatformProcess::IsProcRunning(ProcessHandle) && !IsZenProcessUsingEffectivePort(EffectiveListenPort))
 		{
-			UE_LOG(LogZenServiceInstance, Warning, TEXT("Failed to request shutdown for running service using port %u (Pid: %u)"), EffectiveListenPort, ServicePid);
-			return false;
+			return true;
 		}
-		return true;
+		UE_LOG(LogZenServiceInstance, Warning, TEXT("Failed to request shutdown for running service using port %u (Pid: %u)"), EffectiveListenPort, ServicePid);
+		return false;
 	}
+
 	uint64 ZenShutdownWaitStartTime = FPlatformTime::Cycles64();
 	while (FPlatformProcess::IsProcRunning(ProcessHandle) || IsZenProcessUsingEffectivePort(EffectiveListenPort))
 	{
