Ensure we don't use the messagebox when presenting ZenServer launch messages in standalone programs.

diff --git a/Engine/Source/Developer/Zen/Private/ZenServerInterface.cpp b/Engine/Source/Developer/Zen/Private/ZenServerInterface.cpp
--- a/Engine/Source/Developer/Zen/Private/ZenServerInterface.cpp
+++ b/Engine/Source/Developer/Zen/Private/ZenServerInterface.cpp
@@ -350,7 +350,9 @@ GetZenVersion(const FString& UtilityPath, const FString& ServicePath, const FStr
 static void
 PromptUserToSyncInTreeVersion(const FString& ServerFilePath)
 {
-	if (FApp::IsUnattended())
+#if !IS_PROGRAM
+	if (FApp::IsUnattended() || IsRunningCommandlet() || GIsRunningUnattendedScript)
+#endif
 	{
 		// Just log as there is no one to show a message
 		UE_LOG(LogZenServiceInstance, Display, TEXT("ZenServer can not verify installation. Please make sure your source installation in properly synced at '%s'"), *FPaths::GetPath(ServerFilePath));
@@ -651,7 +653,9 @@ ApplyProcessLifetimeOverride(bool& bLimitProcessLifetime)
 static void
 PromptUserUnableToDetermineValidDataPath()
 {
-	if (FApp::IsUnattended())
+#if !IS_PROGRAM
+	if (FApp::IsUnattended() || IsRunningCommandlet() || GIsRunningUnattendedScript)
+#endif
 	{
 		// Just log as there is no one to show a message
 		UE_LOG(LogZenServiceInstance, Warning, TEXT("ZenServer is unable to determine a valid data path"));
@@ -666,7 +670,9 @@ PromptUserUnableToDetermineValidDataPath()
 static void
 PromptUserAboutInvalidValidDataPathConfiguration(const FString& UsedDataPath)
 {
-	if (FApp::IsUnattended())
+#if !IS_PROGRAM
+	if (FApp::IsUnattended() || IsRunningCommandlet() || GIsRunningUnattendedScript)
+#endif
 	{
 		// Just log as there is no one to show a message
 		UE_LOG(LogZenServiceInstance, Warning, TEXT("ZenServer has detected invalid data path configuration. Falling back to '%s'"), *UsedDataPath);
@@ -2259,25 +2265,25 @@ FZenServiceInstance::AutoLaunch(const FServiceAutoLaunchSettings& InSettings, FS
 			{
 				if (!IsLockFileLocked(*LockFilePath))
 				{
-					if (FApp::IsUnattended())
+#if !IS_PROGRAM
+					if (FApp::IsUnattended() || IsRunningCommandlet() || GIsRunningUnattendedScript)
+#endif
 					{
 						// Just log as there is no one to show a message
 						UE_LOG(LogZenServiceInstance, Warning, TEXT("ZenServer did not launch in the expected duration"));
 						return false;
 					}
-					else
-					{
-						FText ZenLaunchFailurePromptTitle = NSLOCTEXT("Zen", "Zen_LaunchFailurePromptTitle", "Failed to launch");
-
-						FFormatNamedArguments FormatArguments;
-						FString LogFilePath = FPaths::Combine(InSettings.DataPath, TEXT("logs"), TEXT("zenserver.log"));
-						FPaths::MakePlatformFilename(LogFilePath);
-						FormatArguments.Add(TEXT("LogFilePath"), FText::FromString(LogFilePath));
-						FText ZenLaunchFailurePromptText = FText::Format(NSLOCTEXT("Zen", "Zen_LaunchFailurePromptText", "ZenServer failed to launch. This process will now exit. Please check the ZenServer log file for details:\n{LogFilePath}"), FormatArguments);
-						FPlatformMisc::MessageBoxExt(EAppMsgType::Ok, *ZenLaunchFailurePromptText.ToString(), *ZenLaunchFailurePromptTitle.ToString());
-						FPlatformMisc::RequestExit(true);
-						return false;
-					}
+
+					FText ZenLaunchFailurePromptTitle = NSLOCTEXT("Zen", "Zen_LaunchFailurePromptTitle", "Failed to launch");
+
+					FFormatNamedArguments FormatArguments;
+					FString LogFilePath = FPaths::Combine(InSettings.DataPath, TEXT("logs"), TEXT("zenserver.log"));
+					FPaths::MakePlatformFilename(LogFilePath);
+					FormatArguments.Add(TEXT("LogFilePath"), FText::FromString(LogFilePath));
+					FText ZenLaunchFailurePromptText = FText::Format(NSLOCTEXT("Zen", "Zen_LaunchFailurePromptText", "ZenServer failed to launch. This process will now exit. Please check the ZenServer log file for details:\n{LogFilePath}"), FormatArguments);
+					FPlatformMisc::MessageBoxExt(EAppMsgType::Ok, *ZenLaunchFailurePromptText.ToString(), *ZenLaunchFailurePromptTitle.ToString());
+					FPlatformMisc::RequestExit(true);
+					return false;
 				}
 				// Note that the dialog may not show up when zenserver is needed early in the launch cycle, but this will at least ensure
 				// the splash screen is refreshed with the appropriate text status message.
@@ -2285,7 +2291,8 @@ FZenServiceInstance::AutoLaunch(const FServiceAutoLaunchSettings& InSettings, FS
 				UE_LOG(LogZenServiceInstance, Display, TEXT("Waiting for ZenServer to be ready..."));
 				DurationPhase = EWaitDurationPhase::Medium;
 			}
-			else if (!FApp::IsUnattended() && ZenWaitDuration > 20.0 && (DurationPhase == EWaitDurationPhase::Medium))
+#if !IS_PROGRAM
+			else if (!(FApp::IsUnattended() || IsRunningCommandlet() || GIsRunningUnattendedScript) && ZenWaitDuration > 20.0 && (DurationPhase == EWaitDurationPhase::Medium))
 			{
 				FText ZenLongWaitPromptTitle = NSLOCTEXT("Zen", "Zen_LongWaitPromptTitle", "Wait for ZenServer?");
 				FText ZenLongWaitPromptText = NSLOCTEXT("Zen", "Zen_LongWaitPromptText", "ZenServer is taking a long time to launch. It may be performing maintenance. Keep waiting?");
@@ -2296,6 +2303,7 @@ FZenServiceInstance::AutoLaunch(const FServiceAutoLaunchSettings& InSettings, FS
 				}
 				DurationPhase = EWaitDurationPhase::Long;
 			}
+#endif
 
 			if (WaitForZenReadySlowTask.ShouldCancel())
 			{
