Move remote/ConnectExisting ZenServer connections to port 8558 by default.

diff --git a/Engine/Config/BaseEngine.ini b/Engine/Config/BaseEngine.ini
--- a/Engine/Config/BaseEngine.ini
+++ b/Engine/Config/BaseEngine.ini
@@ -2389,7 +2389,7 @@ ExtraArgs=--http asio --gc-cache-duration-seconds 1209600 --gc-interval-seconds
 
 [Zen.ConnectExisting]
 HostName=localhost
-Port=1337
+Port=8558
 
 ; Each key defines settings for an Unreal Cloud DDC server.
 ; The key is used in the ServerID field of a Cloud node in a graph for DDC.
diff --git a/Engine/Source/Developer/Zen/Private/ZenServerInterface.cpp b/Engine/Source/Developer/Zen/Private/ZenServerInterface.cpp
--- a/Engine/Source/Developer/Zen/Private/ZenServerInterface.cpp
+++ b/Engine/Source/Developer/Zen/Private/ZenServerInterface.cpp
@@ -861,7 +861,7 @@ FServiceSettings::ReadFromURL(FStringView InstanceURL)
 	else
 	{
 		ConnectExistingSettings.HostName = InstanceURL;
-		ConnectExistingSettings.Port = 1337;
+		ConnectExistingSettings.Port = 8558;
 	}
 }
 
@@ -902,7 +902,7 @@ FServiceSettings::TryApplyAutoLaunchOverride()
 		SettingsVariant.Emplace<FServiceConnectSettings>();
 		FServiceConnectSettings& ConnectExistingSettings = SettingsVariant.Get<FServiceConnectSettings>();
 		ConnectExistingSettings.HostName = TEXT("[::1]");
-		ConnectExistingSettings.Port = 1337;
+		ConnectExistingSettings.Port = 8558;
 		return true;
 	}
 
@@ -921,7 +921,7 @@ FServiceSettings::TryApplyAutoLaunchOverride()
 		else
 		{
 			ConnectExistingSettings.HostName = Host;
-			ConnectExistingSettings.Port = 1337;
+			ConnectExistingSettings.Port = 8558;
 		}
 
 		return true;
diff --git a/Engine/Source/Developer/Zen/Public/Experimental/ZenServerInterface.h b/Engine/Source/Developer/Zen/Public/Experimental/ZenServerInterface.h
--- a/Engine/Source/Developer/Zen/Public/Experimental/ZenServerInterface.h
+++ b/Engine/Source/Developer/Zen/Public/Experimental/ZenServerInterface.h
@@ -32,7 +32,7 @@ namespace UE::Zen
 struct FServiceConnectSettings
 {
 	FString HostName;
-	uint16 Port = 1337;
+	uint16 Port = 8558;
 };
 
 struct FServiceAutoLaunchSettings
