Fixes for zen stats in latest zenserver release.

diff --git a/Engine/Source/Developer/Zen/Private/ZenServerInterface.cpp b/Engine/Source/Developer/Zen/Private/ZenServerInterface.cpp
--- a/Engine/Source/Developer/Zen/Private/ZenServerInterface.cpp
+++ b/Engine/Source/Developer/Zen/Private/ZenServerInterface.cpp
@@ -2584,17 +2584,6 @@ FZenServiceInstance::GatherAnalytics(TArray<FAnalyticsEventAttribute>& Attribute
 		Attributes.Emplace(MoveTemp(AttrName), ZenCacheStats.General.HitRatio);
 	}
 
-	{
-		FString AttrName = BaseName + TEXT("Cache_UpstreamHits");
-		Attributes.Emplace(MoveTemp(AttrName), ZenCacheStats.General.UpstreamHits);
-	}
-
-	{
-		FString AttrName = BaseName + TEXT("Cache_UpstreamRatio");
-		Attributes.Emplace(MoveTemp(AttrName), ZenCacheStats.General.UpstreamRatio);
-	}
-
-
 	{
 		FString AttrName = BaseName + TEXT("Cache_Cas_Hits");
 		Attributes.Emplace(MoveTemp(AttrName), ZenCacheStats.General.CidHits);
@@ -2612,7 +2601,7 @@ FZenServiceInstance::GatherAnalytics(TArray<FAnalyticsEventAttribute>& Attribute
 
 	{
 		FString AttrName = BaseName + TEXT("Cache_Requests");
-		Attributes.Emplace(MoveTemp(AttrName), ZenCacheStats.General.RequestCount);
+		Attributes.Emplace(MoveTemp(AttrName), ZenCacheStats.Request.Count);
 	}
 
 	{
diff --git a/Engine/Source/Developer/Zen/Private/ZenStatistics.cpp b/Engine/Source/Developer/Zen/Private/ZenStatistics.cpp
--- a/Engine/Source/Developer/Zen/Private/ZenStatistics.cpp
+++ b/Engine/Source/Developer/Zen/Private/ZenStatistics.cpp
@@ -39,12 +39,9 @@ bool LoadFromCompactBinary(FCbFieldView Field, FZenCacheStats::FGeneralStats& Ou
 	bOk = LoadFromCompactBinary(Field["misses"], OutValue.Misses) & bOk;
 	bOk = LoadFromCompactBinary(Field["writes"], OutValue.Writes) & bOk;
 	bOk = LoadFromCompactBinary(Field["hit_ratio"], OutValue.HitRatio) & bOk;
-	bOk = LoadFromCompactBinary(Field["upstream_hits"], OutValue.UpstreamHits) & bOk;
-	bOk = LoadFromCompactBinary(Field["upstream_ratio"], OutValue.UpstreamRatio) & bOk;
 	bOk = LoadFromCompactBinary(Field["cidhits"], OutValue.CidHits) & bOk;
 	bOk = LoadFromCompactBinary(Field["cidmisses"], OutValue.CidMisses) & bOk;
 	bOk = LoadFromCompactBinary(Field["cidwrites"], OutValue.CidWrites) & bOk;
-	bOk = LoadFromCompactBinary(Field["requestcount"], OutValue.RequestCount) & bOk;
 	bOk = LoadFromCompactBinary(Field["badrequestcount"], OutValue.BadRequestCount) & bOk;
 	return bOk;
 }
@@ -113,8 +110,13 @@ bool LoadFromCompactBinary(FCbFieldView Field, FZenCacheStats& OutValue)
 	bool bOk = Field.IsObject();
 	bOk = LoadFromCompactBinary(Field["cache"], OutValue.General) & bOk;
 	bOk = LoadFromCompactBinary(Field["requests"], OutValue.Request) & bOk;
-	bOk = LoadFromCompactBinary(Field["upstream"], OutValue.Upstream) & bOk;
-	bOk = LoadFromCompactBinary(Field["upstream_gets"], OutValue.UpstreamRequest) & bOk;
+
+	FCbFieldView UpstreamFieldView = Field["upstream"];
+	if (UpstreamFieldView.IsObject())
+	{
+		bOk = LoadFromCompactBinary(UpstreamFieldView, OutValue.Upstream) & bOk;
+		bOk = LoadFromCompactBinary(Field["upstream_gets"], OutValue.UpstreamRequest) & bOk;
+	}
 	bOk = LoadFromCompactBinary(Field["cid"], OutValue.CID) & bOk;
 	return bOk;
 }
diff --git a/Engine/Source/Developer/Zen/Public/Experimental/ZenStatistics.h b/Engine/Source/Developer/Zen/Public/Experimental/ZenStatistics.h
--- a/Engine/Source/Developer/Zen/Public/Experimental/ZenStatistics.h
+++ b/Engine/Source/Developer/Zen/Public/Experimental/ZenStatistics.h
@@ -42,12 +42,9 @@ struct FZenCacheStats
 		int64 Misses = 0;
 		int64 Writes = 0;
 		double HitRatio = 0.0;
-		int64 UpstreamHits = 0;
-		double UpstreamRatio = 0.0;
 		int64 CidHits = 0;
 		int64 CidMisses = 0;
 		int64 CidWrites = 0;
-		int64 RequestCount = 0;
 		int64 BadRequestCount = 0;
 	};
 
diff --git a/Engine/Source/Programs/ZenDashboard/Private/SZenCacheStatistics.cpp b/Engine/Source/Programs/ZenDashboard/Private/SZenCacheStatistics.cpp
--- a/Engine/Source/Programs/ZenDashboard/Private/SZenCacheStatistics.cpp
+++ b/Engine/Source/Programs/ZenDashboard/Private/SZenCacheStatistics.cpp
@@ -174,7 +174,7 @@ TSharedRef<SWidget> SZenCacheStatistics::GetGridPanel()
 	[
 		SNew(STextBlock)
 		.Margin(FMargin(ColumnMargin, RowMargin))
-		.Text_Lambda([bHaveStats, Requests = ZenStats.General.RequestCount]
+		.Text_Lambda([bHaveStats, Requests = ZenStats.Request.Count]
 		{
 			if (bHaveStats)
 			{
