Change timeout for zenserver shutdown from 15 seconds to 25 seconds.  Experienced zenserver shutdown that was taking 16 seconds.

diff --git a/Engine/Source/Developer/Zen/Private/ZenServerInterface.cpp b/Engine/Source/Developer/Zen/Private/ZenServerInterface.cpp
--- a/Engine/Source/Developer/Zen/Private/ZenServerInterface.cpp
+++ b/Engine/Source/Developer/Zen/Private/ZenServerInterface.cpp
@@ -1236,7 +1236,7 @@ static bool AttemptShutdownUsingExecutablePathOnly(const ZenServerState& ServerS
 }
 
 static bool
-ShutdownRunningServiceUsingExecutablePath(const TCHAR* ExecutablePath, double MaximumWaitDurationSeconds = 15.0)
+ShutdownRunningServiceUsingExecutablePath(const TCHAR* ExecutablePath, double MaximumWaitDurationSeconds = 25.0)
 {
 	uint32 ServicePid;
 	if (!IsZenProcessActive(ExecutablePath, &ServicePid))
@@ -1360,7 +1360,7 @@ static bool AttemptShutdownUsingPortOnly(const ZenServerState& ServerState, uint
 }
 
 static bool
-ShutdownRunningServiceUsingEffectivePort(uint16 EffectiveListenPort, double MaximumWaitDurationSeconds = 15.0)
+ShutdownRunningServiceUsingEffectivePort(uint16 EffectiveListenPort, double MaximumWaitDurationSeconds = 25.0)
 {
 	if (!IsZenProcessUsingEffectivePort(EffectiveListenPort))
 	{
diff --git a/Engine/Source/Developer/Zen/Public/Experimental/ZenServerInterface.h b/Engine/Source/Developer/Zen/Public/Experimental/ZenServerInterface.h
--- a/Engine/Source/Developer/Zen/Public/Experimental/ZenServerInterface.h
+++ b/Engine/Source/Developer/Zen/Public/Experimental/ZenServerInterface.h
@@ -114,7 +114,7 @@ private:
 
 UE_API bool IsLocalServiceRunning(const TCHAR* DataPath, uint16* OutPort = nullptr);
 UE_API FProcHandle StartLocalService(const FZenLocalServiceRunContext& Context, const TCHAR* TransientArgs = nullptr);
-UE_API bool StopLocalService(const TCHAR* DataPath, double MaximumWaitDurationSeconds = 15.0);
+UE_API bool StopLocalService(const TCHAR* DataPath, double MaximumWaitDurationSeconds = 25.0);
 
 UE_API FString GetLocalServiceInstallPath();
 UE_API FString GetLocalServiceInstallVersion(bool bDetailed = true);
