Fixed unreachable code errors.

diff --git a/Engine/Source/Developer/Zen/Private/ZenServerInterface.cpp b/Engine/Source/Developer/Zen/Private/ZenServerInterface.cpp
--- a/Engine/Source/Developer/Zen/Private/ZenServerInterface.cpp
+++ b/Engine/Source/Developer/Zen/Private/ZenServerInterface.cpp
@@ -351,17 +351,18 @@ static void
 PromptUserToSyncInTreeVersion(const FString& ServerFilePath)
 {
 #if !IS_PROGRAM
-	if (FApp::IsUnattended() || IsRunningCommandlet() || GIsRunningUnattendedScript)
+	if (!FApp::IsUnattended() && !IsRunningCommandlet() && !GIsRunningUnattendedScript)
+	{
+		FText ZenSyncSourcePromptTitle = NSLOCTEXT("Zen", "Zen_SyncSourcePromptTitle", "Failed to launch");
+		FText ZenSyncSourcePromptText = FText::Format(NSLOCTEXT("Zen", "Zen_SyncSourcePromptText", "ZenServer can not verify installation. Please make sure your source installation in properly synced at '{0}'"), FText::FromString(FPaths::GetPath(ServerFilePath)));
+		FPlatformMisc::MessageBoxExt(EAppMsgType::Ok, *ZenSyncSourcePromptText.ToString(), *ZenSyncSourcePromptTitle.ToString());
+	}
+	else
 #endif
 	{
 		// Just log as there is no one to show a message
 		UE_LOG(LogZenServiceInstance, Display, TEXT("ZenServer can not verify installation. Please make sure your source installation in properly synced at '%s'"), *FPaths::GetPath(ServerFilePath));
-		return;
 	}
-
-	FText ZenSyncSourcePromptTitle = NSLOCTEXT("Zen", "Zen_SyncSourcePromptTitle", "Failed to launch");
-	FText ZenSyncSourcePromptText = FText::Format(NSLOCTEXT("Zen", "Zen_SyncSourcePromptText", "ZenServer can not verify installation. Please make sure your source installation in properly synced at '{0}'"), FText::FromString(FPaths::GetPath(ServerFilePath)));
-	FPlatformMisc::MessageBoxExt(EAppMsgType::Ok, *ZenSyncSourcePromptText.ToString(), *ZenSyncSourcePromptTitle.ToString());
 }
 
 static bool
@@ -654,35 +655,38 @@ static void
 PromptUserUnableToDetermineValidDataPath()
 {
 #if !IS_PROGRAM
-	if (FApp::IsUnattended() || IsRunningCommandlet() || GIsRunningUnattendedScript)
+	if (!FApp::IsUnattended() && !IsRunningCommandlet() && !GIsRunningUnattendedScript)
+	{
+		FString LogDirPath = FPaths::ConvertRelativePathToFull(FPaths::ProjectLogDir());
+		FText ZenInvalidDataPathPromptTitle = NSLOCTEXT("Zen", "Zen_InvalidDataPathPromptTitle", "No Valid Data Path Configuration");
+		FText ZenInvalidDataPathPromptText = FText::Format(NSLOCTEXT("Zen", "Zen_InvalidDataPathPromptText", "ZenServer can not determine a valid data path.\nPlease check the log in '{0}' for details.\nUpdate your configuration and restart."), FText::FromString(LogDirPath));
+		FPlatformMisc::MessageBoxExt(EAppMsgType::Ok, *ZenInvalidDataPathPromptText.ToString(), *ZenInvalidDataPathPromptTitle.ToString());
+	}
+	else
 #endif
 	{
 		// Just log as there is no one to show a message
 		UE_LOG(LogZenServiceInstance, Warning, TEXT("ZenServer is unable to determine a valid data path"));
-		return;
 	}
-	FString LogDirPath = FPaths::ConvertRelativePathToFull(FPaths::ProjectLogDir());
-	FText ZenInvalidDataPathPromptTitle = NSLOCTEXT("Zen", "Zen_InvalidDataPathPromptTitle", "No Valid Data Path Configuration");
-	FText ZenInvalidDataPathPromptText = FText::Format(NSLOCTEXT("Zen", "Zen_InvalidDataPathPromptText", "ZenServer can not determine a valid data path.\nPlease check the log in '{0}' for details.\nUpdate your configuration and restart."), FText::FromString(LogDirPath));
-	FPlatformMisc::MessageBoxExt(EAppMsgType::Ok, *ZenInvalidDataPathPromptText.ToString(), *ZenInvalidDataPathPromptTitle.ToString());
 }
 
 static void
 PromptUserAboutInvalidValidDataPathConfiguration(const FString& UsedDataPath)
 {
 #if !IS_PROGRAM
-	if (FApp::IsUnattended() || IsRunningCommandlet() || GIsRunningUnattendedScript)
+	if (!FApp::IsUnattended() && !IsRunningCommandlet() && !GIsRunningUnattendedScript)
+	{
+		FString LogDirPath = FPaths::ConvertRelativePathToFull(FPaths::ProjectLogDir());
+		FText ZenInvalidValidDataPathConfigurationPromptTitle = NSLOCTEXT("Zen", "Zen_InvalidValidDataPathConfigurationPromptTitle", "Invalid Data Paths");
+		FText ZenInvalidValidDataPathConfigurationPromptText = FText::Format(NSLOCTEXT("Zen", "Zen_InvalidValidDataPathConfigurationPromptText", "ZenServer has detected invalid data path configuration.\nPlease check the log in '{0}' for details.\n\nFalling back to using '{1}' as data path."), FText::FromString(LogDirPath), FText::FromString(UsedDataPath));
+		FPlatformMisc::MessageBoxExt(EAppMsgType::Ok, *ZenInvalidValidDataPathConfigurationPromptText.ToString(), *ZenInvalidValidDataPathConfigurationPromptTitle.ToString());
+	}
+	else
 #endif
 	{
 		// Just log as there is no one to show a message
 		UE_LOG(LogZenServiceInstance, Warning, TEXT("ZenServer has detected invalid data path configuration. Falling back to '%s'"), *UsedDataPath);
-		return;
 	}
-
-	FString LogDirPath = FPaths::ConvertRelativePathToFull(FPaths::ProjectLogDir());
-	FText ZenInvalidValidDataPathConfigurationPromptTitle = NSLOCTEXT("Zen", "Zen_InvalidValidDataPathConfigurationPromptTitle", "Invalid Data Paths");
-	FText ZenInvalidValidDataPathConfigurationPromptText = FText::Format(NSLOCTEXT("Zen", "Zen_InvalidValidDataPathConfigurationPromptText", "ZenServer has detected invalid data path configuration.\nPlease check the log in '{0}' for details.\n\nFalling back to using '{1}' as data path."), FText::FromString(LogDirPath), FText::FromString(UsedDataPath));
-	FPlatformMisc::MessageBoxExt(EAppMsgType::Ok, *ZenInvalidValidDataPathConfigurationPromptText.ToString(), *ZenInvalidValidDataPathConfigurationPromptTitle.ToString());
 }
 
 #if PLATFORM_WINDOWS
@@ -1986,52 +1990,55 @@ static void
 PromptUserToStopRunningServerInstanceForUpdate(const FString& ServerFilePath)
 {
 #if !IS_PROGRAM
-	if (FApp::IsUnattended() || IsRunningCommandlet() || GIsRunningUnattendedScript)
+	if (!FApp::IsUnattended() && !IsRunningCommandlet() && !GIsRunningUnattendedScript)
+	{
+		FText ZenUpdatePromptTitle = NSLOCTEXT("Zen", "Zen_UpdatePromptTitle", "Update required");
+		FText ZenUpdatePromptText = FText::Format(NSLOCTEXT("Zen", "Zen_UpdatePromptText", "ZenServer needs to be updated to a new version. Please shut down Unreal Editor and any tools that are using the ZenServer at '{0}'"), FText::FromString(ServerFilePath));
+		FPlatformMisc::MessageBoxExt(EAppMsgType::Ok, *ZenUpdatePromptText.ToString(), *ZenUpdatePromptTitle.ToString());
+	}
+	else
 #endif
 	{
 		// Just log as there is no one to show a message
 		UE_LOG(LogZenServiceInstance, Display, TEXT("ZenServer needs to be updated to a new version. Please shut down any tools that are using the ZenServer at '%s'"), *ServerFilePath);
-		return;
 	}
-
-	FText ZenUpdatePromptTitle = NSLOCTEXT("Zen", "Zen_UpdatePromptTitle", "Update required");
-	FText ZenUpdatePromptText = FText::Format(NSLOCTEXT("Zen", "Zen_UpdatePromptText", "ZenServer needs to be updated to a new version. Please shut down Unreal Editor and any tools that are using the ZenServer at '{0}'"), FText::FromString(ServerFilePath));
-	FPlatformMisc::MessageBoxExt(EAppMsgType::Ok, *ZenUpdatePromptText.ToString(), *ZenUpdatePromptTitle.ToString());
 }
 
 static void
 PromptUserOfLockedDataFolder(const FString& DataPath)
 {
 #if !IS_PROGRAM
-	if (FApp::IsUnattended() || IsRunningCommandlet() || GIsRunningUnattendedScript)
+	if (!FApp::IsUnattended() && !IsRunningCommandlet() && !GIsRunningUnattendedScript)
+	{
+		FText ZenLaunchFailurePromptTitle = NSLOCTEXT("Zen", "Zen_NonLocalProcessUsesDataDirPromptTitle", "Failed to launch");
+		FText ZenLaunchFailurePromptText = FText::Format(NSLOCTEXT("Zen", "Zen_NonLocalProcessUsesDataDirPromptText", "ZenServer Failed to auto launch, an unknown process is locking the data folder '{0}'"), FText::FromString(DataPath));
+		FPlatformMisc::MessageBoxExt(EAppMsgType::Ok, *ZenLaunchFailurePromptText.ToString(), *ZenLaunchFailurePromptTitle.ToString());
+	}
+	else
 #endif
 	{
 		// Just log as there is no one to show a message
 		UE_LOG(LogZenServiceInstance, Warning, TEXT("ZenServer Failed to auto launch, an unknown process is locking the data folder '%s'"), *DataPath);
-		return;
 	}
-
-	FText ZenLaunchFailurePromptTitle = NSLOCTEXT("Zen", "Zen_NonLocalProcessUsesDataDirPromptTitle", "Failed to launch");
-	FText ZenLaunchFailurePromptText = FText::Format(NSLOCTEXT("Zen", "Zen_NonLocalProcessUsesDataDirPromptText", "ZenServer Failed to auto launch, an unknown process is locking the data folder '{0}'"), FText::FromString(DataPath));
-	FPlatformMisc::MessageBoxExt(EAppMsgType::Ok, *ZenLaunchFailurePromptText.ToString(), *ZenLaunchFailurePromptTitle.ToString());
 }
 
 static void
 PromptUserOfFailedShutDownOfExistingProcess(uint16 Port)
 {
 #if !IS_PROGRAM
-	if (FApp::IsUnattended() || IsRunningCommandlet() || GIsRunningUnattendedScript)
+	if (!FApp::IsUnattended() && !IsRunningCommandlet() && !GIsRunningUnattendedScript)
+	{
+		FText ZenLaunchFailurePromptTitle = NSLOCTEXT("Zen", "Zen_ShutdownFailurePromptTitle", "Failed to launch");
+		FText ZenLaunchFailurePromptText = FText::Format(NSLOCTEXT("Zen", "Zen_ShutdownFailurePromptText", "ZenServer Failed to auto launch, failed to shut down currently running service using port '{0}'"), FText::AsNumber(Port, &FNumberFormattingOptions::DefaultNoGrouping()));
+		FPlatformMisc::MessageBoxExt(EAppMsgType::Ok, *ZenLaunchFailurePromptText.ToString(), *ZenLaunchFailurePromptTitle.ToString());
+	}
+	else
 #endif
 	{
 		// Just log as there is no one to show a message
 		UE_LOG(LogZenServiceInstance, Warning, TEXT("ZenServer Failed to auto launch, failed to shut down currently running service using port %u"), Port);
 		FPlatformMisc::RequestExit(true);
-		return;
 	}
-
-	FText ZenLaunchFailurePromptTitle = NSLOCTEXT("Zen", "Zen_ShutdownFailurePromptTitle", "Failed to launch");
-	FText ZenLaunchFailurePromptText = FText::Format(NSLOCTEXT("Zen", "Zen_ShutdownFailurePromptText", "ZenServer Failed to auto launch, failed to shut down currently running service using port '{0}'"), FText::AsNumber(Port, &FNumberFormattingOptions::DefaultNoGrouping()));
-	FPlatformMisc::MessageBoxExt(EAppMsgType::Ok, *ZenLaunchFailurePromptText.ToString(), *ZenLaunchFailurePromptTitle.ToString());
 }
 
 
@@ -2273,24 +2280,26 @@ FZenServiceInstance::AutoLaunch(const FServiceAutoLaunchSettings& InSettings, FS
 				if (!IsLockFileLocked(*LockFilePath))
 				{
 #if !IS_PROGRAM
-					if (FApp::IsUnattended() || IsRunningCommandlet() || GIsRunningUnattendedScript)
+					if (!FApp::IsUnattended() && !IsRunningCommandlet() && !GIsRunningUnattendedScript)
+					{
+						FText ZenLaunchFailurePromptTitle = NSLOCTEXT("Zen", "Zen_LaunchFailurePromptTitle", "Failed to launch");
+
+						FFormatNamedArguments FormatArguments;
+						FString LogFilePath = FPaths::Combine(InSettings.DataPath, TEXT("logs"), TEXT("zenserver.log"));
+						FPaths::MakePlatformFilename(LogFilePath);
+						FormatArguments.Add(TEXT("LogFilePath"), FText::FromString(LogFilePath));
+						FText ZenLaunchFailurePromptText = FText::Format(NSLOCTEXT("Zen", "Zen_LaunchFailurePromptText", "ZenServer failed to launch. This process will now exit. Please check the ZenServer log file for details:\n{LogFilePath}"), FormatArguments);
+						FPlatformMisc::MessageBoxExt(EAppMsgType::Ok, *ZenLaunchFailurePromptText.ToString(), *ZenLaunchFailurePromptTitle.ToString());
+						FPlatformMisc::RequestExit(true);
+						return false;
+					}
+					else
 #endif
 					{
 						// Just log as there is no one to show a message
 						UE_LOG(LogZenServiceInstance, Warning, TEXT("ZenServer did not launch in the expected duration"));
 						return false;
 					}
-
-					FText ZenLaunchFailurePromptTitle = NSLOCTEXT("Zen", "Zen_LaunchFailurePromptTitle", "Failed to launch");
-
-					FFormatNamedArguments FormatArguments;
-					FString LogFilePath = FPaths::Combine(InSettings.DataPath, TEXT("logs"), TEXT("zenserver.log"));
-					FPaths::MakePlatformFilename(LogFilePath);
-					FormatArguments.Add(TEXT("LogFilePath"), FText::FromString(LogFilePath));
-					FText ZenLaunchFailurePromptText = FText::Format(NSLOCTEXT("Zen", "Zen_LaunchFailurePromptText", "ZenServer failed to launch. This process will now exit. Please check the ZenServer log file for details:\n{LogFilePath}"), FormatArguments);
-					FPlatformMisc::MessageBoxExt(EAppMsgType::Ok, *ZenLaunchFailurePromptText.ToString(), *ZenLaunchFailurePromptTitle.ToString());
-					FPlatformMisc::RequestExit(true);
-					return false;
 				}
 				// Note that the dialog may not show up when zenserver is needed early in the launch cycle, but this will at least ensure
 				// the splash screen is refreshed with the appropriate text status message.
diff --git a/Engine/Source/Runtime/Core/Private/Internationalization/PolyglotTextSource.cpp b/Engine/Source/Runtime/Core/Private/Internationalization/PolyglotTextSource.cpp
--- a/Engine/Source/Runtime/Core/Private/Internationalization/PolyglotTextSource.cpp
+++ b/Engine/Source/Runtime/Core/Private/Internationalization/PolyglotTextSource.cpp
@@ -19,9 +19,9 @@ bool FPolyglotTextSource::GetNativeCultureName(const ELocalizedTextSourceCategor
 
 	if (FCultureInfo* CultureInfo = AvailableCultureInfo.Find(InCategory))
 	{
-		for (const auto& NativeCulturePair : CultureInfo->NativeCultures)
+		if (CultureInfo->NativeCultures.Num() > 0)
 		{
-			OutNativeCultureName = NativeCulturePair.Key;
+			OutNativeCultureName = CultureInfo->NativeCultures.CreateConstIterator()->Key;
 			return true;
 		}
 	}
diff --git a/Engine/Source/Runtime/CoreUObject/Private/Serialization/DerivedData.cpp b/Engine/Source/Runtime/CoreUObject/Private/Serialization/DerivedData.cpp
--- a/Engine/Source/Runtime/CoreUObject/Private/Serialization/DerivedData.cpp
+++ b/Engine/Source/Runtime/CoreUObject/Private/Serialization/DerivedData.cpp
@@ -312,7 +312,7 @@ void FDerivedData::Serialize(FArchive& Ar, UObject* Owner)
 
 		///////////////////////////////////////////////////////////////////////////////////////////////////////////////
 		// TEMPORARY SERIALIZATION AS BULK DATA
-
+#if 1
 		FCookedData LocalCookedData = Linker->AddDerivedData(*this).CookedData;
 		LocalCookedData.ChunkId[8] = 0;
 		LocalCookedData.ChunkId[9] = 0;
@@ -362,8 +362,9 @@ void FDerivedData::Serialize(FArchive& Ar, UObject* Owner)
 
 		// TEMPORARY SERIALIZATION AS BULK DATA
 		///////////////////////////////////////////////////////////////////////////////////////////////////////////////
-
+#else
 		return Linker->AddDerivedData(*this).CookedData.Serialize(Ar);
+#endif
 	}
 #endif
 
diff --git a/Engine/Source/Runtime/CoreUObject/Private/UObject/UObjectClusters.cpp b/Engine/Source/Runtime/CoreUObject/Private/UObject/UObjectClusters.cpp
--- a/Engine/Source/Runtime/CoreUObject/Private/UObject/UObjectClusters.cpp
+++ b/Engine/Source/Runtime/CoreUObject/Private/UObject/UObjectClusters.cpp
@@ -497,16 +497,7 @@ void FindStaleClusters(const TArray<FString>& Args)
 			UObject* ClusterRootObject = static_cast<UObject*>(ObjectItem->Object);
 			FReferenceChainSearch SearchRefs(ClusterRootObject, EReferenceChainSearchMode::ExternalOnly);
 			
-			bool bReferenced = false;
-			if (SearchRefs.GetReferenceChains().Num() > 0)
-			{
-				for (const FReferenceChainSearch::FReferenceChain* ReferenceChain : SearchRefs.GetReferenceChains())
-				{
-					bReferenced = true;
-					break;
-				}
-			}
-			if (!bReferenced)
+			if (SearchRefs.GetReferenceChains().Num() == 0)
 			{
 				NumStaleClusters++;
 				UE_LOG(LogObj, Display, TEXT("Cluster %s has no external references:"), *ClusterRootObject->GetFullName());
