Add missing changes for CL 30501078
For local Zen DDC connections show the path where data is store in the Cache Statistics tab for DDC instead of local ip address
For Zen DDC connection show the storage size

diff --git a/Engine/Source/Developer/Zen/Private/ZenServerInterface.cpp b/Engine/Source/Developer/Zen/Private/ZenServerInterface.cpp
--- a/Engine/Source/Developer/Zen/Private/ZenServerInterface.cpp
+++ b/Engine/Source/Developer/Zen/Private/ZenServerInterface.cpp
@@ -1797,7 +1797,16 @@ FZenServiceInstance::~FZenServiceInstance()
 {
 }
 
-bool 
+const FString FZenServiceInstance::GetPath() const
+{
+	if (Settings.IsAutoLaunch())
+	{
+		return Settings.SettingsVariant.Get<FServiceAutoLaunchSettings>().DataPath;
+	}
+	return GetURL();
+}
+
+bool
 FZenServiceInstance::IsServiceRunning()
 {
 	return !Settings.IsAutoLaunch() || bHasLaunchedLocal;
