Add missing changes for CL 30501142 (and 30501566)
For local Zen DDC connections show the path where data is store in the Cache Statistics tab for DDC instead of local ip address
For Zen DDC connection show the storage size

diff --git a/Engine/Source/Developer/Zen/Private/ZenServerInterface.cpp b/Engine/Source/Developer/Zen/Private/ZenServerInterface.cpp
--- a/Engine/Source/Developer/Zen/Private/ZenServerInterface.cpp
+++ b/Engine/Source/Developer/Zen/Private/ZenServerInterface.cpp
@@ -8,6 +8,7 @@
 
 #include "AnalyticsEventAttribute.h"
 #include "Async/Async.h"
+#include "Async/UniqueLock.h"
 #include "Dom/JsonValue.h"
 #include "HAL/FileManager.h"
 #include "HAL/Platform.h"
@@ -2323,7 +2324,6 @@ FZenServiceInstance::AutoLaunch(const FServiceAutoLaunchSettings& InSettings, FS
 				FText ZenLongWaitPromptText = NSLOCTEXT("Zen", "Zen_LongWaitPromptText", "ZenServer is taking a long time to launch. It may be performing maintenance. Keep waiting?");
 				if (FPlatformMisc::MessageBoxExt(EAppMsgType::YesNo, *ZenLongWaitPromptText.ToString(), *ZenLongWaitPromptTitle.ToString()) == EAppReturnType::No)
 				{
-					FPlatformMisc::RequestExit(true);
 					return false;
 				}
 				DurationPhase = EWaitDurationPhase::Long;
@@ -2332,7 +2332,6 @@ FZenServiceInstance::AutoLaunch(const FServiceAutoLaunchSettings& InSettings, FS
 
 			if (WaitForZenReadySlowTask.ShouldCancel())
 			{
-				FPlatformMisc::RequestExit(true);
 				return false;
 			}
 			FPlatformProcess::Sleep(0.1f);
@@ -2344,31 +2343,30 @@ FZenServiceInstance::AutoLaunch(const FServiceAutoLaunchSettings& InSettings, FS
 bool 
 FZenServiceInstance::GetCacheStats(FZenCacheStats& Stats)
 {
-	check(IsInGameThread());
-
-	// If we've already requested stats and they are ready then grab them
-	if ( CacheStatsRequest.IsReady() == true )
+	UE::Zen::FZenHttpRequest* Request = nullptr;
 	{
-		LastCacheStats		= CacheStatsRequest.Get();
-		LastCacheStatsTime	= FPlatformTime::Cycles64();
+		TUniqueLock Lock(LastCacheStatsMutex);
+		// If we've already requested stats and they are ready then grab them
+		if ( CacheStatsRequest.IsReady() == true )
+		{
+			LastCacheStats		= CacheStatsRequest.Get();
+			LastCacheStatsTime	= FPlatformTime::Cycles64();
 
-		CacheStatsRequest.Reset();
-	}
+			CacheStatsRequest.Reset();
+		}
 	
-	// Make a copy of the last updated stats
-	Stats = LastCacheStats;
+		// Make a copy of the last updated stats
+		Stats = LastCacheStats;
 
-	const uint64 CurrentTime = FPlatformTime::Cycles64();
-	constexpr double MinTimeBetweenRequestsInSeconds = 0.5;
-	const double DeltaTimeInSeconds = FPlatformTime::ToSeconds64(CurrentTime - LastCacheStatsTime);
+		const uint64 CurrentTime = FPlatformTime::Cycles64();
+		constexpr double MinTimeBetweenRequestsInSeconds = 0.5;
+		const double DeltaTimeInSeconds = FPlatformTime::ToSeconds64(CurrentTime - LastCacheStatsTime);
+
+		if (CacheStatsRequest.IsValid() || DeltaTimeInSeconds <= MinTimeBetweenRequestsInSeconds)
+		{
+			return Stats.bIsValid;
+		}
 
-	if (!CacheStatsRequest.IsValid() && DeltaTimeInSeconds > MinTimeBetweenRequestsInSeconds)
-	{
-#if WITH_EDITOR
-		EAsyncExecution ThreadPool = EAsyncExecution::LargeThreadPool;
-#else
-		EAsyncExecution ThreadPool = EAsyncExecution::ThreadPool;
-#endif
 		if (!CacheStatsHttpRequest.IsValid())
 		{
 			TStringBuilder<128> ZenDomain;
@@ -2376,26 +2374,34 @@ FZenServiceInstance::GetCacheStats(FZenCacheStats& Stats)
 			CacheStatsHttpRequest = MakePimpl<FZenHttpRequest>(ZenDomain.ToString(), false);
 		}
 
+		Request = CacheStatsHttpRequest.Get();
+	}
+
+#if WITH_EDITOR
+	EAsyncExecution ThreadPool = EAsyncExecution::LargeThreadPool;
+#else
+	EAsyncExecution ThreadPool = EAsyncExecution::ThreadPool;
+#endif
 		// We've not got any requests in flight and we've met a given time requirement for requests
-		CacheStatsRequest = Async(ThreadPool, [this]
+	CacheStatsRequest = Async(ThreadPool, [Request]
 		{
-			UE::Zen::FZenHttpRequest& Request = *CacheStatsHttpRequest.Get();
-			Request.Reset();
+		check(Request != nullptr);
+		Request->Reset();
 
 			TArray64<uint8> GetBuffer;
-			FZenHttpRequest::Result Result = Request.PerformBlockingDownload(TEXTVIEW("/stats/z$"), &GetBuffer, Zen::EContentType::CbObject);
+		FZenHttpRequest::Result Result = Request->PerformBlockingDownload(TEXTVIEW("/stats/z$"), &GetBuffer, Zen::EContentType::CbObject);
 
 			FZenCacheStats Stats;
 
-			if (Result == Zen::FZenHttpRequest::Result::Success && Request.GetResponseCode() == 200)
+		if (Result == Zen::FZenHttpRequest::Result::Success && Request->GetResponseCode() == 200)
 			{
+			UE_LOG(LogZenServiceInstance, Display, TEXT("FZenServiceInstance::GetCacheStats (updating stats from response)"));
 				FCbFieldView RootView(GetBuffer.GetData());
 				Stats.bIsValid = LoadFromCompactBinary(RootView, Stats);
 			}
 
 			return Stats;
 		});
-	}
 
 	return Stats.bIsValid;
 }
@@ -2403,50 +2409,55 @@ FZenServiceInstance::GetCacheStats(FZenCacheStats& Stats)
 bool 
 FZenServiceInstance::GetProjectStats(FZenProjectStats& Stats)
 {
-	check(IsInGameThread());
-
-	// If we've already requested stats and they are ready then grab them
-	if ( ProjectStatsRequest.IsReady() == true )
+	UE::Zen::FZenHttpRequest* Request = nullptr;
 	{
-		LastProjectStats		= ProjectStatsRequest.Get();
-		LastProjectStatsTime	= FPlatformTime::Cycles64();
+		TUniqueLock Lock(LastProjectStatsMutex);
+		// If we've already requested stats and they are ready then grab them
+		if ( ProjectStatsRequest.IsReady() == true )
+		{
+			LastProjectStats		= ProjectStatsRequest.Get();
+			LastProjectStatsTime	= FPlatformTime::Cycles64();
 
-		ProjectStatsRequest.Reset();
-	}
+			ProjectStatsRequest.Reset();
+		}
 	
-	// Make a copy of the last updated stats
-	Stats = LastProjectStats;
+		// Make a copy of the last updated stats
+		Stats = LastProjectStats;
 
-	const uint64 CurrentTime = FPlatformTime::Cycles64();
-	constexpr double MinTimeBetweenRequestsInSeconds = 0.5;
-	const double DeltaTimeInSeconds = FPlatformTime::ToSeconds64(CurrentTime - LastProjectStatsTime);
+		const uint64 CurrentTime = FPlatformTime::Cycles64();
+		constexpr double MinTimeBetweenRequestsInSeconds = 0.5;
+		const double DeltaTimeInSeconds = FPlatformTime::ToSeconds64(CurrentTime - LastProjectStatsTime);
 
-	if (!ProjectStatsRequest.IsValid() && DeltaTimeInSeconds > MinTimeBetweenRequestsInSeconds)
-	{
-#if WITH_EDITOR
-		EAsyncExecution ThreadPool = EAsyncExecution::LargeThreadPool;
-#else
-		EAsyncExecution ThreadPool = EAsyncExecution::ThreadPool;
-#endif
+		if (ProjectStatsRequest.IsValid() || DeltaTimeInSeconds <= MinTimeBetweenRequestsInSeconds)
+		{
+			return Stats.bIsValid;
+		}
 		if (!ProjectStatsHttpRequest.IsValid())
 		{
 			TStringBuilder<128> ZenDomain;
 			ZenDomain << HostName << TEXT(":") << Port;
 			ProjectStatsHttpRequest = MakePimpl<FZenHttpRequest>(ZenDomain.ToString(), false);
 		}
+		Request = ProjectStatsHttpRequest.Get();
+	}
 
-		// We've not got any requests in flight and we've met a given time requirement for requests
-		ProjectStatsRequest = Async(ThreadPool, [this]
-		{
-			UE::Zen::FZenHttpRequest& Request = *ProjectStatsHttpRequest.Get();
-			Request.Reset();
+#if WITH_EDITOR
+	EAsyncExecution ThreadPool = EAsyncExecution::LargeThreadPool;
+#else
+	EAsyncExecution ThreadPool = EAsyncExecution::ThreadPool;
+#endif
+			// We've not got any requests in flight and we've met a given time requirement for requests
+	ProjectStatsRequest = Async(ThreadPool, [Request]
+	{
+		check(Request);
+		Request->Reset();
 
 			TArray64<uint8> GetBuffer;
-			FZenHttpRequest::Result Result = Request.PerformBlockingDownload(TEXTVIEW("/stats/prj"), &GetBuffer, Zen::EContentType::CbObject);
+		FZenHttpRequest::Result Result = Request->PerformBlockingDownload(TEXTVIEW("/stats/prj"), &GetBuffer, Zen::EContentType::CbObject);
 
 			FZenProjectStats Stats;
 
-			if (Result == Zen::FZenHttpRequest::Result::Success && Request.GetResponseCode() == 200)
+		if (Result == Zen::FZenHttpRequest::Result::Success && Request->GetResponseCode() == 200)
 			{
 				FCbFieldView RootView(GetBuffer.GetData());
 				Stats.bIsValid = LoadFromCompactBinary(RootView, Stats);
@@ -2454,7 +2465,6 @@ FZenServiceInstance::GetProjectStats(FZenProjectStats& Stats)
 
 			return Stats;
 		});
-	}
 
 	return Stats.bIsValid;
 }
