Add logging when signalling zenserver event for shutdown fails

diff --git a/Engine/Source/Developer/Zen/Private/ZenServerInterface.cpp b/Engine/Source/Developer/Zen/Private/ZenServerInterface.cpp
--- a/Engine/Source/Developer/Zen/Private/ZenServerInterface.cpp
+++ b/Engine/Source/Developer/Zen/Private/ZenServerInterface.cpp
@@ -1160,14 +1160,24 @@ static bool
 RequestZenShutdownOnEffectivePort(uint16 EffectiveListenPort)
 {
 #if PLATFORM_WINDOWS
-	HANDLE Handle = OpenEventW(EVENT_MODIFY_STATE, false, *WriteToWideString<64>(WIDETEXT("Zen_"), EffectiveListenPort, WIDETEXT("_Shutdown")));
-	if (Handle != NULL)
+	TAnsiStringBuilder<64> EventPath;
+	EventPath << "Zen_" << EffectiveListenPort << "_Shutdown";
+	HANDLE Handle = OpenEventA(EVENT_MODIFY_STATE, false, *EventPath);
+	if (Handle == NULL)
 	{
-		ON_SCOPE_EXIT{ CloseHandle(Handle); };
-		BOOL OK = SetEvent(Handle);
-		return OK;
+		DWORD err = GetLastError();
+		UE_LOG(LogZenServiceInstance, Warning, TEXT("Failed opening named event '%hs' (err: %d)"), *EventPath, err);
+		return false;
 	}
-	return false;
+	ON_SCOPE_EXIT{ CloseHandle(Handle); };
+	BOOL OK = SetEvent(Handle);
+	if (!OK)
+	{
+		DWORD err = GetLastError();
+		UE_LOG(LogZenServiceInstance, Warning, TEXT("Failed signaling named event '%hs' (err: %d)"), *EventPath, err);
+		return false;
+	}
+	return true;
 #elif PLATFORM_UNIX || PLATFORM_MAC
 	TAnsiStringBuilder<64> EventPath;
 	EventPath << "/tmp/Zen_" << EffectiveListenPort << "_Shutdown";
@@ -1175,16 +1185,25 @@ RequestZenShutdownOnEffectivePort(uint16 EffectiveListenPort)
 	key_t IpcKey = ftok(EventPath.ToString(), 1);
 	if (IpcKey < 0)
 	{
+		int err = errno;
+		UE_LOG(LogZenServiceInstance, Warning, TEXT("Failed resolving ipc key for named event '%hs' (err: %d)"), *EventPath, err);
 		return false;
 	}
 
 	int Semaphore = semget(IpcKey, 1, 0600);
 	if (Semaphore < 0)
 	{
+		int err = errno;
+		UE_LOG(LogZenServiceInstance, Warning, TEXT("Failed opening named event '%hs' (err: %d)"), *EventPath, err);
 		return false;
 	}
 
-	semctl(Semaphore, 0, SETVAL, 0);
+	if (semctl(Semaphore, 0, SETVAL, 0) < 0)
+	{
+		int err = errno;
+		UE_LOG(LogZenServiceInstance, Warning, TEXT("Failed signaling named event '%hs' (err: %d)"), *EventPath, err);
+		return false;
+	}
 	return true;
 #else
 	static_assert(false, "Missing implementation for Zen named shutdown events");
