Remove launch zenserver as admin fallback

diff --git a/Engine/Source/Developer/Zen/Private/ZenServerInterface.cpp b/Engine/Source/Developer/Zen/Private/ZenServerInterface.cpp
--- a/Engine/Source/Developer/Zen/Private/ZenServerInterface.cpp
+++ b/Engine/Source/Developer/Zen/Private/ZenServerInterface.cpp
@@ -1616,7 +1616,8 @@ StartLocalService(const FZenLocalServiceRunContext& Context, const TCHAR* Transi
 	FString PlatformWorkingDirectory = Context.GetWorkingDirectory();
 	FPaths::MakePlatformFilename(PlatformWorkingDirectory);
 	{
-		// Attempt non-elevated launch
+		// We could switch to FPlatformProcess::CreateProc for Windows as well if we are able to add the CREATE_BREAKAWAY_FROM_JOB flag
+		// as that is needed on CI to stop Horde from terminating the zenserver process
 		STARTUPINFO StartupInfo = {
 			sizeof(STARTUPINFO),
 			NULL, NULL, NULL,
@@ -1640,24 +1641,9 @@ StartLocalService(const FZenLocalServiceRunContext& Context, const TCHAR* Transi
 			::CloseHandle(ProcInfo.hThread);
 			Proc = FProcHandle(ProcInfo.hProcess);
 		}
-
-	}
-	if (!Proc.IsValid())
-	{
-		// Fall back to elevated launch
-		SHELLEXECUTEINFO ShellExecuteInfo;
-		ZeroMemory(&ShellExecuteInfo, sizeof(ShellExecuteInfo));
-		ShellExecuteInfo.cbSize = sizeof(ShellExecuteInfo);
-		ShellExecuteInfo.fMask = SEE_MASK_UNICODE | SEE_MASK_NOCLOSEPROCESS;
-		ShellExecuteInfo.lpFile = *PlatformExecutable;
-		ShellExecuteInfo.lpDirectory = *PlatformWorkingDirectory;
-		ShellExecuteInfo.lpVerb = TEXT("runas");
-		ShellExecuteInfo.nShow = Context.GetShowConsole() ? SW_SHOWMINNOACTIVE : SW_HIDE;
-		ShellExecuteInfo.lpParameters = *Parms;
-
-		if (ShellExecuteEx(&ShellExecuteInfo))
+		else
 		{
-			Proc = FProcHandle(ShellExecuteInfo.hProcess);
+			UE_LOG(LogZenServiceInstance, Warning, TEXT("Failed launching %s status: %d."), *CommandLine, GetLastError());
 		}
 	}
 #else
