Validate running processes for zenserver by using native functions directly.
We need to control how the process handle is opened.

diff --git a/Engine/Source/Developer/Zen/Private/ZenServerInterface.cpp b/Engine/Source/Developer/Zen/Private/ZenServerInterface.cpp
--- a/Engine/Source/Developer/Zen/Private/ZenServerInterface.cpp
+++ b/Engine/Source/Developer/Zen/Private/ZenServerInterface.cpp
@@ -54,10 +54,93 @@ namespace UE::Zen
 
 DEFINE_LOG_CATEGORY_STATIC(LogZenServiceInstance, Log, All);
 
+// Native functions to interact with a process using a process id
+// We don't use UE's own OpenProcess as they try to open processes with PROCESS_ALL_ACCESS
+static bool NativeIsProcessRunning(uint32 Pid)
+{
+	if (Pid == 0)
+	{
+		return false;
+	}
+#if PLATFORM_WINDOWS
+	HANDLE Handle = OpenProcess(PROCESS_QUERY_LIMITED_INFORMATION, 0, (DWORD)Pid);
+
+	if (!Handle)
+	{
+		DWORD Error = GetLastError();
+
+		if (Error == ERROR_INVALID_PARAMETER)
+		{
+			return false;
+		}
+		return false;
+	}
+	ON_SCOPE_EXIT{ CloseHandle(Handle); };
+
+	bool  bStillActive = true;
+	DWORD ExitCode = 0;
+	if (0 != GetExitCodeProcess(Handle, &ExitCode))
+	{
+		bStillActive = ExitCode == STILL_ACTIVE;
+	}
+
+	return bStillActive;
+#elif PLATFORM_UNIX || PLATFORM_MAC
+	return (kill(pid_t(Pid), 0) == 0);
+#endif
+}
+
+static bool NativeTerminate(uint32 Pid)
+{
+#if PLATFORM_WINDOWS
+	HANDLE Handle = OpenProcess(PROCESS_TERMINATE | SYNCHRONIZE, 0, (DWORD)Pid);
+	if (Handle == NULL)
+	{
+		DWORD Error = GetLastError();
+
+		if (Error != ERROR_INVALID_PARAMETER)
+		{
+			// LOG
+			return false;
+		}
+		return true;
+	}
+	ON_SCOPE_EXIT{ CloseHandle(Handle); };
+
+	BOOL bTerminated = TerminateProcess(Handle, 0);
+	if (!bTerminated)
+	{
+		DWORD Error = GetLastError();
+		// LOG
+		return false;
+	}
+	DWORD WaitResult = WaitForSingleObject(Handle, INFINITE);
+	BOOL bSuccess = (WaitResult != WAIT_OBJECT_0) && (WaitResult != WAIT_ABANDONED_0);
+	if (!bSuccess)
+	{
+		DWORD Error = GetLastError();
+		// LOG
+		return false;
+	}
+#elif PLATFORM_UNIX || PLATFORM_MAC
+	int Res = kill(pid_t(Pid), SIGKILL);
+	if (Res != 0)
+	{
+		int err = errno;
+		if (err != ESRCH)
+		{
+			// LOG
+			return false;
+		}
+	}
+#endif
+	return true;
+}
+
 class ZenServerState
 {
 public:
-	ZenServerState(bool ReadOnly = true);
+	ZenServerState(bool ReadOnly);
 	~ZenServerState();
 
 	struct ZenServerEntry
@@ -68,10 +151,10 @@ public:
 		std::atomic<uint32> Pid;
 		std::atomic<uint16> DesiredListenPort;
 		std::atomic<uint16> Flags;
-		uint8				  SessionId[12];
+		uint8				SessionId[12];
 		std::atomic<uint32> SponsorPids[8];
 		std::atomic<uint16> EffectiveListenPort;
-		uint8				  Padding[10];
+		uint8				Padding[10];
 
 		enum class FlagsEnum : uint16
 		{
@@ -79,7 +162,7 @@ public:
 			kIsReady = 1 << 1,
 		};
 
-		bool AddSponsorProcess(uint32_t PidToAdd);
+		bool AddSponsorProcess(uint32 PidToAdd);
 	};
 	static_assert(sizeof(ZenServerEntry) == 64);
 
@@ -187,9 +270,14 @@ const ZenServerState::ZenServerEntry* ZenServerState::LookupByDesiredListenPortI
 
 	for (int i = 0; i < m_MaxEntryCount; ++i)
 	{
-		if (m_Data[i].DesiredListenPort == Port)
+		if (m_Data[i].DesiredListenPort.load(std::memory_order_relaxed) == Port)
 		{
-			return &m_Data[i];
+			const ZenServerState::ZenServerEntry* Entry = &m_Data[i];
+			if (NativeIsProcessRunning((uint32)Entry->Pid.load(std::memory_order_relaxed)))
+			{
+				return Entry;
+			}
+			return nullptr;
 		}
 	}
 
@@ -216,9 +304,14 @@ const ZenServerState::ZenServerEntry* ZenServerState::LookupByEffectiveListenPor
 
 	for (int i = 0; i < m_MaxEntryCount; ++i)
 	{
-		if (m_Data[i].EffectiveListenPort == Port)
+		if (m_Data[i].EffectiveListenPort.load(std::memory_order_relaxed) == Port)
 		{
-			return &m_Data[i];
+			const ZenServerState::ZenServerEntry* Entry = &m_Data[i];
+			if (NativeIsProcessRunning((uint32)Entry->Pid.load(std::memory_order_relaxed)))
+			{
+				return Entry;
+			}
+			return nullptr;
 		}
 	}
 
@@ -245,7 +338,7 @@ const ZenServerState::ZenServerEntry* ZenServerState::LookupByPid(uint32 Pid) co
 
 	for (int i = 0; i < m_MaxEntryCount; ++i)
 	{
-		if (m_Data[i].Pid == Pid)
+		if (m_Data[i].Pid.load(std::memory_order_relaxed) == Pid)
 		{
 			return &m_Data[i];
 		}
@@ -255,13 +348,13 @@ const ZenServerState::ZenServerEntry* ZenServerState::LookupByPid(uint32 Pid) co
 }
 
 bool
-ZenServerState::ZenServerEntry::AddSponsorProcess(uint32_t PidToAdd)
+ZenServerState::ZenServerEntry::AddSponsorProcess(uint32 PidToAdd)
 {
-	for (std::atomic<uint32_t>& PidEntry : SponsorPids)
+	for (std::atomic<uint32>& PidEntry : SponsorPids)
 	{
 		if (PidEntry.load(std::memory_order_relaxed) == 0)
 		{
-			uint32_t Expected = 0;
+			uint32 Expected = 0;
 			if (PidEntry.compare_exchange_strong(Expected, PidToAdd))
 			{
 				// Success!
@@ -1230,7 +1323,7 @@ static bool AttemptShutdownUsingExecutablePathOnly(const ZenServerState& ServerS
 		return false;
 	}
 
-	uint16 EffectiveListenPort = Entry->EffectiveListenPort;
+	uint16 EffectiveListenPort = Entry->EffectiveListenPort.load(std::memory_order_relaxed);
 	if (!RequestZenShutdownOnEffectivePort(EffectiveListenPort))
 	{
 		if (!IsZenProcessActive(ExecutablePath, &ServicePid))
@@ -1271,70 +1364,38 @@ ShutdownRunningServiceUsingExecutablePath(const TCHAR* ExecutablePath, double Ma
 
 	UE_LOG(LogZenServiceInstance, Display, TEXT("Waiting for running instance using executable '%s' to shut down"), ExecutablePath);
 
-	FProcHandle ProcessHandle = FPlatformProcess::OpenProcess(ServicePid);
-	ON_SCOPE_EXIT{ FPlatformProcess::CloseProc(ProcessHandle); };
-
-	const ZenServerState ServerState;
+	const ZenServerState ServerState(/* ReadOnly */true);
 	const ZenServerState::ZenServerEntry* Entry = ServerState.LookupByPid(ServicePid);
 	if (!Entry)
 	{
-		if (!IsZenProcessActive(ExecutablePath, nullptr))
+		UE_LOG(LogZenServiceInstance, Warning, TEXT("Can't find server state for running service for executable '%s' (Pid: %u), attempting hard terminate"), ExecutablePath, ServicePid);
+		if (NativeTerminate(ServicePid))
 		{
 			return true;
 		}
-		if (ProcessHandle.IsValid())
-		{
-			UE_LOG(LogZenServiceInstance, Warning, TEXT("Can't find server state for running service for executable '%s' (Pid: %u), attempting hard terminate"), ExecutablePath, ServicePid);
-			FPlatformProcess::TerminateProc(ProcessHandle);
-			FPlatformProcess::CloseProc(ProcessHandle);
-
-			if (!IsZenProcessActive(ExecutablePath, nullptr))
-			{
-				return true;
-			}
-
-			UE_LOG(LogZenServiceInstance, Warning, TEXT("Failed to terminate running service for executable '%s' (Pid: %u)"), ExecutablePath, ServicePid);
-			return false;
-		}
-		UE_LOG(LogZenServiceInstance, Warning, TEXT("Can't find server state for running service for executable '%s' (Pid: %u)"), ExecutablePath, ServicePid);
+		UE_LOG(LogZenServiceInstance, Warning, TEXT("Failed to terminate running service for executable '%s' (Pid: %u)"), ExecutablePath, ServicePid);
 		return false;
 	}
 
-	uint16 EffectiveListenPort = Entry->EffectiveListenPort;
-
-	if (!ProcessHandle.IsValid())
-	{
-		if (!IsZenProcessActive(ExecutablePath, nullptr))
-		{
-			return true;
-		}
-
-		UE_LOG(LogZenServiceInstance, Warning, TEXT("Failed to open handle for running service for executable '%s' (Pid: %u, Port: %u), attempting blind shut down"), ExecutablePath, ServicePid, EffectiveListenPort);
-		return AttemptShutdownUsingExecutablePathOnly(ServerState, ExecutablePath, MaximumWaitDurationSeconds);
-	}
+	uint16 EffectiveListenPort = Entry->EffectiveListenPort.load(std::memory_order_relaxed);
 
 	if (!RequestZenShutdownOnEffectivePort(EffectiveListenPort))
 	{
-		if (!FPlatformProcess::IsProcRunning(ProcessHandle) && !IsZenProcessActive(ExecutablePath, nullptr))
+		if (!IsZenProcessActive(ExecutablePath, nullptr))
 		{
 			return true;
 		}
-
 		UE_LOG(LogZenServiceInstance, Warning, TEXT("Failed to request shutdown for running service for executable '%s' (Pid: %u, Port: %u), attempting hard terminate"), ExecutablePath, ServicePid, EffectiveListenPort);
-		FPlatformProcess::TerminateProc(ProcessHandle);
-		FPlatformProcess::CloseProc(ProcessHandle);
-
-		if (!IsZenProcessActive(ExecutablePath, nullptr))
+		if (NativeTerminate(ServicePid))
 		{
 			return true;
 		}
-
 		UE_LOG(LogZenServiceInstance, Warning, TEXT("Failed to terminate running service for executable '%s' (Pid: %u, Port: %u)"), ExecutablePath, ServicePid, EffectiveListenPort);
 		return false;
 	}
 
 	uint64 ZenShutdownWaitStartTime = FPlatformTime::Cycles64();
-	while (FPlatformProcess::IsProcRunning(ProcessHandle) || IsZenProcessActive(ExecutablePath, nullptr))
+	while (IsZenProcessActive(ExecutablePath, nullptr))
 	{
 		double ZenShutdownWaitDuration = FPlatformTime::ToSeconds64(FPlatformTime::Cycles64() - ZenShutdownWaitStartTime);
 		if (ZenShutdownWaitDuration < MaximumWaitDurationSeconds)
@@ -1393,7 +1454,7 @@ ShutdownRunningServiceUsingEffectivePort(uint16 EffectiveListenPort, double Maxi
 
 	UE_LOG(LogZenServiceInstance, Display, TEXT("Waiting for running instance using port %u to shut down"), EffectiveListenPort);
 
-	const ZenServerState ServerState;
+	const ZenServerState ServerState(/*ReadOnly*/true);
 	const ZenServerState::ZenServerEntry* Entry = ServerState.LookupByEffectiveListenPort(EffectiveListenPort);
 	if (!Entry)
 	{
@@ -1406,23 +1467,10 @@ ShutdownRunningServiceUsingEffectivePort(uint16 EffectiveListenPort, double Maxi
 		return AttemptShutdownUsingPortOnly(ServerState, EffectiveListenPort, MaximumWaitDurationSeconds);
 	}
 
-	uint32 ServicePid = Entry->Pid;
-	FProcHandle ProcessHandle = FPlatformProcess::OpenProcess(ServicePid);
-	ON_SCOPE_EXIT{ FPlatformProcess::CloseProc(ProcessHandle); };
-
-	if (!ProcessHandle.IsValid())
-	{
-		if (!IsZenProcessUsingEffectivePort(EffectiveListenPort))
-		{
-			return true;
-		}
-		UE_LOG(LogZenServiceInstance, Warning, TEXT("Failed to open handle for running service using port %u (Pid: %u), attempting blind shut down"), EffectiveListenPort, ServicePid);
-		return AttemptShutdownUsingPortOnly(ServerState, EffectiveListenPort, MaximumWaitDurationSeconds);
-	}
-
+	uint32 ServicePid = Entry->Pid.load(std::memory_order_relaxed);
 	if (!RequestZenShutdownOnEffectivePort(EffectiveListenPort))
 	{
-		if (!FPlatformProcess::IsProcRunning(ProcessHandle) && !IsZenProcessUsingEffectivePort(EffectiveListenPort))
+		if (!NativeIsProcessRunning(ServicePid) && !IsZenProcessUsingEffectivePort(EffectiveListenPort))
 		{
 			return true;
 		}
@@ -1431,7 +1479,7 @@ ShutdownRunningServiceUsingEffectivePort(uint16 EffectiveListenPort, double Maxi
 	}
 
 	uint64 ZenShutdownWaitStartTime = FPlatformTime::Cycles64();
-	while (FPlatformProcess::IsProcRunning(ProcessHandle) || IsZenProcessUsingEffectivePort(EffectiveListenPort))
+	while (NativeIsProcessRunning(ServicePid) || IsZenProcessUsingEffectivePort(EffectiveListenPort))
 	{
 		double ZenShutdownWaitDuration = FPlatformTime::ToSeconds64(FPlatformTime::Cycles64() - ZenShutdownWaitStartTime);
 		if (ZenShutdownWaitDuration < MaximumWaitDurationSeconds)
@@ -1443,6 +1491,12 @@ ShutdownRunningServiceUsingEffectivePort(uint16 EffectiveListenPort, double Maxi
 			UE_LOG(LogZenServiceInstance, Warning, TEXT("Timed out waiting for shutdown of running service using port %u (Pid: %u)"), EffectiveListenPort, ServicePid);
 			return false;
 		}
+		Entry = ServerState.LookupByEffectiveListenPort(EffectiveListenPort);
+		if (!Entry)
+		{
+			return AttemptShutdownUsingPortOnly(ServerState, EffectiveListenPort, MaximumWaitDurationSeconds - ZenShutdownWaitDuration);
+		}
+		ServicePid = Entry->Pid.load(std::memory_order_relaxed);
 	}
 	return true;
 }
@@ -1876,7 +1930,7 @@ FZenServiceInstance::TryRecovery()
 			bool bShutdownExistingInstance = true;
 			std::atomic<uint32> PreviousSponsorPids[UE_ARRAY_COUNT(ZenServerState::ZenServerEntry::SponsorPids)];
 			{
-				const ZenServerState ServerState;
+				const ZenServerState ServerState(/* ReadOnly */true);
 				const ZenServerState::ZenServerEntry* Entry = ServerState.LookupByEffectiveListenPort(Port);
 				if (Entry)
 				{
@@ -1980,7 +2034,7 @@ FZenServiceInstance::Initialize()
 			bHasLaunchedLocal = AutoLaunch(Settings.SettingsVariant.Get<FServiceAutoLaunchSettings>(), *ExecutableInstallPath, HostName, Port);
 			if (bHasLaunchedLocal)
 			{
-				const ZenServerState State;
+				const ZenServerState State(/*ReadOnly*/true);
 				const ZenServerState::ZenServerEntry* RunningEntry = State.LookupByEffectiveListenPort(Port);
 				if (RunningEntry != nullptr)
 				{
@@ -2146,7 +2200,7 @@ FZenServiceInstance::AutoLaunch(const FServiceAutoLaunchSettings& InSettings, FS
 
 	if (LockFileState.IsReady)
 	{
-		const ZenServerState State;
+		const ZenServerState State(/*ReadOnly*/true);
 		if (State.LookupByPid(LockFileState.ProcessId) == nullptr)
 		{
 			PromptUserOfLockedDataFolder(*InSettings.DataPath);
@@ -2182,11 +2236,11 @@ FZenServiceInstance::AutoLaunch(const FServiceAutoLaunchSettings& InSettings, FS
 	}
 	else
 	{
-		const ZenServerState State;
+		const ZenServerState State(/*ReadOnly*/true);
 		const ZenServerState::ZenServerEntry* RunningEntry = State.LookupByDesiredListenPort(InSettings.DesiredPort);
 		if (RunningEntry != nullptr)
 		{
-			ShutdownEffectivePort = RunningEntry->EffectiveListenPort;
+			ShutdownEffectivePort = RunningEntry->EffectiveListenPort.load(std::memory_order_relaxed);
 			UE_LOG(LogZenServiceInstance, Log, TEXT("Found existing instance running on port %u with different data directory, will attempt shut down"), ShutdownEffectivePort);
 		}
 	}
