Fix for zenserver shutdown and start logic bug.  The "return" within the while NativeIsProcessRunning || IsZenProcessUsingEffectivePort loop creates a situation where we can exit the loop before the process is completely gone, meaning the lockfile could still be present.

diff --git a/Engine/Source/Developer/Zen/Private/ZenServerInterface.cpp b/Engine/Source/Developer/Zen/Private/ZenServerInterface.cpp
--- a/Engine/Source/Developer/Zen/Private/ZenServerInterface.cpp
+++ b/Engine/Source/Developer/Zen/Private/ZenServerInterface.cpp
@@ -1492,11 +1492,10 @@ ShutdownRunningServiceUsingEffectivePort(uint16 EffectiveListenPort, double Maxi
 			return false;
 		}
 		Entry = ServerState.LookupByEffectiveListenPort(EffectiveListenPort);
-		if (!Entry)
+		if (Entry)
 		{
-			return AttemptShutdownUsingPortOnly(ServerState, EffectiveListenPort, MaximumWaitDurationSeconds - ZenShutdownWaitDuration);
+			ServicePid = Entry->Pid.load(std::memory_order_relaxed);
 		}
-		ServicePid = Entry->Pid.load(std::memory_order_relaxed);
 	}
 	return true;
 }
