Fix non-editor (program) use cases of Zen AutoLaunch to ensure that they load the editor settings config before accessing it.

diff --git a/Engine/Source/Developer/Zen/Private/ZenServerInterface.cpp b/Engine/Source/Developer/Zen/Private/ZenServerInterface.cpp
--- a/Engine/Source/Developer/Zen/Private/ZenServerInterface.cpp
+++ b/Engine/Source/Developer/Zen/Private/ZenServerInterface.cpp
@@ -18,6 +18,7 @@
 #include "Logging/LogScopedCategoryAndVerbosityOverride.h"
 #include "Misc/App.h"
 #include "Misc/ConfigCacheIni.h"
+#include "Misc/ConfigContext.h"
 #include "Misc/FileHelper.h"
 #include "Misc/MonitoredProcess.h"
 #include "Misc/Paths.h"
@@ -524,6 +525,18 @@ AttemptFileCopyWithRetries(const TCHAR* Dst, const TCHAR* Src, double RetryDurat
 	return false;
 }
 
+static void EnsureEditorSettingsConfigLoaded()
+{
+#if !WITH_EDITOR
+	if (GEditorSettingsIni.IsEmpty())
+	{
+		FConfigContext Context = FConfigContext::ReadIntoGConfig();
+		Context.GeneratedConfigDir = FPaths::EngineEditorSettingsDir();
+		Context.Load(TEXT("EditorSettings"), GEditorSettingsIni);
+	}
+#endif
+}
+
 static void
 DetermineLocalDataCachePath(const TCHAR* ConfigSection, FString& DataPath)
 {
@@ -567,6 +580,7 @@ DetermineLocalDataCachePath(const TCHAR* ConfigSection, FString& DataPath)
 	FString DataPathEditorOverrideSetting;
 	if (GConfig->GetString(ConfigSection, TEXT("LocalDataCachePathEditorOverrideSetting"), DataPathEditorOverrideSetting, GEngineIni))
 	{
+		EnsureEditorSettingsConfigLoaded();
 		FString Setting = GConfig->GetStr(TEXT("/Script/UnrealEd.EditorSettings"), *DataPathEditorOverrideSetting, GEditorSettingsIni);
 		if (!Setting.IsEmpty())
 		{
@@ -887,6 +901,7 @@ FServiceSettings::ReadFromConfig()
 			GConfig->GetBool(AutoLaunchConfigSection, TEXT("ShowConsole"), AutoLaunchSettings.bShowConsole, GEngineIni);
 			GConfig->GetBool(AutoLaunchConfigSection, TEXT("LimitProcessLifetime"), AutoLaunchSettings.bLimitProcessLifetime, GEngineIni);
 			ApplyProcessLifetimeOverride(AutoLaunchSettings.bLimitProcessLifetime);
+			EnsureEditorSettingsConfigLoaded();
 			GConfig->GetBool(TEXT("/Script/UnrealEd.CrashReportsPrivacySettings"), TEXT("bSendUnattendedBugReports"), AutoLaunchSettings.bSendUnattendedBugReports, GEditorSettingsIni);
 		}
 	}
