Further hardening of zen service shut down

diff --git a/Engine/Source/Developer/Zen/Private/ZenServerInterface.cpp b/Engine/Source/Developer/Zen/Private/ZenServerInterface.cpp
--- a/Engine/Source/Developer/Zen/Private/ZenServerInterface.cpp
+++ b/Engine/Source/Developer/Zen/Private/ZenServerInterface.cpp
@@ -9,6 +9,7 @@
 #include "AnalyticsEventAttribute.h"
 #include "Async/Async.h"
 #include "Async/UniqueLock.h"
+#include "Containers/AnsiString.h"
 #include "Dom/JsonValue.h"
 #include "HAL/FileManager.h"
 #include "HAL/Platform.h"
@@ -101,7 +102,7 @@ static bool NativeTerminate(uint32 Pid)
 
 		if (Error != ERROR_INVALID_PARAMETER)
 		{
-			// LOG
+			UE_LOG(LogZenServiceInstance, Warning, TEXT("Failed to open running process for terminate %d: %d"), Pid, Error);
 			return false;
 		}
 		return true;
@@ -112,15 +113,18 @@ static bool NativeTerminate(uint32 Pid)
 	if (!bTerminated)
 	{
 		DWORD Error = GetLastError();
-		// LOG
+		UE_LOG(LogZenServiceInstance, Warning, TEXT("Failed to terminate running process %d: %d"), Pid, Error);
 		return false;
 	}
-	DWORD WaitResult = WaitForSingleObject(Handle, INFINITE);
-	BOOL bSuccess = (WaitResult != WAIT_OBJECT_0) && (WaitResult != WAIT_ABANDONED_0);
+	DWORD WaitResult = WaitForSingleObject(Handle, 15000);
+	BOOL bSuccess = (WaitResult == WAIT_OBJECT_0) || (WaitResult == WAIT_ABANDONED_0);
 	if (!bSuccess)
 	{
-		DWORD Error = GetLastError();
-		// LOG
+		if (WaitResult == WAIT_FAILED)
+		{
+			DWORD Error = GetLastError();
+			UE_LOG(LogZenServiceInstance, Warning, TEXT("Failed to wait for terminated process %d: %d"), Pid, Error);
+		}
 		return false;
 	}
 #elif PLATFORM_UNIX || PLATFORM_MAC
@@ -130,7 +134,7 @@ static bool NativeTerminate(uint32 Pid)
 		int err = errno;
 		if (err != ESRCH)
 		{
-			// LOG
+			UE_LOG(LogZenServiceInstance, Warning, TEXT("Failed to terminate running process %d: %d"), Pid, err);
 			return false;
 		}
 	}
@@ -305,9 +309,9 @@ const ZenServerState::ZenServerEntry* ZenServerState::LookupByEffectiveListenPor
 
 	for (int i = 0; i < m_MaxEntryCount; ++i)
 	{
-		if (m_Data[i].EffectiveListenPort.load(std::memory_order_relaxed) == Port)
+		const ZenServerState::ZenServerEntry* Entry = &m_Data[i];
+		if (Entry->EffectiveListenPort.load(std::memory_order_relaxed) == Port)
 		{
-			const ZenServerState::ZenServerEntry* Entry = &m_Data[i];
 			if (NativeIsProcessRunning((uint32)Entry->Pid.load(std::memory_order_relaxed)))
 			{
 				return Entry;
@@ -339,9 +343,13 @@ const ZenServerState::ZenServerEntry* ZenServerState::LookupByPid(uint32 Pid) co
 
 	for (int i = 0; i < m_MaxEntryCount; ++i)
 	{
+		const ZenServerState::ZenServerEntry* Entry = &m_Data[i];
 		if (m_Data[i].Pid.load(std::memory_order_relaxed) == Pid)
 		{
-			return &m_Data[i];
+			if (NativeIsProcessRunning(Pid))
+			{
+				return Entry;
+			}
 		}
 	}
 
@@ -1157,8 +1165,25 @@ ReadCbLockFile(FStringView FileName)
 #endif
 }
 
+FString GetLocalServiceExecutableName()
+{
+	return
+#if PLATFORM_WINDOWS
+		TEXT("zenserver.exe")
+#else
+		TEXT("zenserver")
+#endif
+		;
+}
+
+FString
+GetLocalServiceInstallPath()
+{
+	return FPaths::ConvertRelativePathToFull(FPaths::Combine(FPlatformProcess::ApplicationSettingsDir(), TEXT("Zen\\Install"), GetLocalServiceExecutableName()));
+}
+
 static bool
-IsLockFileLocked(const TCHAR* FileName, bool bAttemptCleanUp=false)
+IsLockFileLocked(const TCHAR* FileName, bool bAttemptCleanUp = false)
 {
 #if PLATFORM_WINDOWS
 	if (bAttemptCleanUp)
@@ -1193,88 +1218,71 @@ IsLockFileLocked(const TCHAR* FileName, bool bAttemptCleanUp=false)
 #endif
 }
 
-static bool
-IsZenProcessUsingEffectivePort(uint16 EffectiveListenPort)
-{
 #if PLATFORM_WINDOWS
-	HANDLE Handle = OpenEventW(READ_CONTROL, false, *WriteToWideString<64>(WIDETEXT("Zen_"), EffectiveListenPort, WIDETEXT("_Shutdown")));
-	if (Handle != NULL)
-	{
-		ON_SCOPE_EXIT{ CloseHandle(Handle); };
-		return true;
-	}
-	return false;
-#elif PLATFORM_UNIX || PLATFORM_MAC
-	TAnsiStringBuilder<64> EventPath;
-	EventPath << "/tmp/Zen_" << EffectiveListenPort << "_Shutdown";
+static FString GetZenProcessShutdownEventName(uint16 EffectiveListenPort)
+{
+	return *WriteToWideString<64>(WIDETEXT("Zen_"), EffectiveListenPort, WIDETEXT("_Shutdown"));
+}
 
-	key_t IpcKey = ftok(EventPath.ToString(), 1);
+static HANDLE OpenNativeEvent(const FString& EventName, DWORD Access)
+{
+	return OpenEventW(Access, false, *EventName);
+}
+#elif PLATFORM_UNIX || PLATFORM_MAC
+static FAnsiString GetZenProcessShutdownEventName(uint16 EffectiveListenPort)
+{
+	return *WriteToAnsiString<64>("/tmp/Zen_", EffectiveListenPort, "_Shutdown");
+}
+static int OpenNativeEvent(const FAnsiString& EventName, int Access)
+{
+	key_t IpcKey = ftok(*EventName, 1);
 	if (IpcKey < 0)
 	{
-		return false;
+		return -1;
 	}
-
-	int Semaphore = semget(IpcKey, 1, 0400);
+	int Semaphore = semget(IpcKey, 1, Access);
 	if (Semaphore < 0)
 	{
-		return false;
+		return -1;
 	}
-	return true;
-#else
-	static_assert(false, "Missing implementation for Zen named shutdown events");
-	return false;
-#endif
+	return Semaphore;
 }
+#endif
 
 static bool
-IsZenProcessActive(const TCHAR* ExecutablePath, uint32* OutPid)
+IsZenProcessUsingEffectivePort(uint16 EffectiveListenPort)
 {
-	FString NormalizedExecutablePath(ExecutablePath);
-	FPaths::NormalizeFilename(NormalizedExecutablePath);
-	FPlatformProcess::FProcEnumerator ProcIter;
-	while (ProcIter.MoveNext())
+#if PLATFORM_WINDOWS
+	HANDLE Handle = OpenNativeEvent(GetZenProcessShutdownEventName(EffectiveListenPort), READ_CONTROL);
+	if (Handle != NULL)
 	{
-		FPlatformProcess::FProcEnumInfo ProcInfo = ProcIter.GetCurrent();
-		FString Candidate = ProcInfo.GetFullPath();
-		FPaths::NormalizeFilename(Candidate);
-		if (Candidate == NormalizedExecutablePath)
-		{
-			if (OutPid)
-			{
-				*OutPid = ProcInfo.GetPID();
-			}
-			return true;
-		}
+		CloseHandle(Handle);
+		return true;
 	}
 	return false;
-}
-
-static bool
-IsZenProcessUsingDataDir(const TCHAR* LockFilePath, LockFileData* OutLockFileData)
-{
-	if (IsLockFileLocked(LockFilePath, true))
+#elif PLATFORM_UNIX || PLATFORM_MAC
+	int Semaphore = OpenNativeEvent(GetZenProcessShutdownEventName(EffectiveListenPort), 0400);
+	if (Semaphore >= 0)
 	{
-		if (OutLockFileData)
-		{
-			// If an instance is running with this data path, check if we can use it and what port it is on
-			*OutLockFileData = ReadCbLockFile(LockFilePath);
-		}
 		return true;
 	}
 	return false;
+#else
+	static_assert(false, "Missing implementation for Zen named shutdown events");
+	return false;
+#endif
 }
 
 static bool
 RequestZenShutdownOnEffectivePort(uint16 EffectiveListenPort)
 {
 #if PLATFORM_WINDOWS
-	TAnsiStringBuilder<64> EventPath;
-	EventPath << "Zen_" << EffectiveListenPort << "_Shutdown";
-	HANDLE Handle = OpenEventA(EVENT_MODIFY_STATE, false, *EventPath);
+	FString ShutdownEventName = GetZenProcessShutdownEventName(EffectiveListenPort);
+	HANDLE Handle = OpenNativeEvent(ShutdownEventName, EVENT_MODIFY_STATE);
 	if (Handle == NULL)
 	{
 		DWORD err = GetLastError();
-		UE_LOG(LogZenServiceInstance, Warning, TEXT("Failed opening named event '%hs' (err: %d)"), *EventPath, err);
+		UE_LOG(LogZenServiceInstance, Warning, TEXT("Failed opening named event '%s' (err: %d)"), *ShutdownEventName, err);
 		return false;
 	}
 	ON_SCOPE_EXIT{ CloseHandle(Handle); };
@@ -1282,34 +1290,25 @@ RequestZenShutdownOnEffectivePort(uint16 EffectiveListenPort)
 	if (!OK)
 	{
 		DWORD err = GetLastError();
-		UE_LOG(LogZenServiceInstance, Warning, TEXT("Failed signaling named event '%hs' (err: %d)"), *EventPath, err);
+		UE_LOG(LogZenServiceInstance, Warning, TEXT("Failed signaling named event '%s' (err: %d)"), *ShutdownEventName, err);
 		return false;
 	}
 	return true;
+	
 #elif PLATFORM_UNIX || PLATFORM_MAC
-	TAnsiStringBuilder<64> EventPath;
-	EventPath << "/tmp/Zen_" << EffectiveListenPort << "_Shutdown";
-
-	key_t IpcKey = ftok(EventPath.ToString(), 1);
-	if (IpcKey < 0)
-	{
-		int err = errno;
-		UE_LOG(LogZenServiceInstance, Warning, TEXT("Failed resolving ipc key for named event '%hs' (err: %d)"), *EventPath, err);
-		return false;
-	}
-
-	int Semaphore = semget(IpcKey, 1, 0600);
+	FAnsiString ShutdownEventName = GetZenProcessShutdownEventName(EffectiveListenPort);
+	int Semaphore = OpenNativeEvent(GetZenProcessShutdownEventName(EffectiveListenPort), 0600);
 	if (Semaphore < 0)
 	{
 		int err = errno;
-		UE_LOG(LogZenServiceInstance, Warning, TEXT("Failed opening named event '%hs' (err: %d)"), *EventPath, err);
+		UE_LOG(LogZenServiceInstance, Warning, TEXT("Failed opening named event '%hs' (err: %d)"), *ShutdownEventName, err);
 		return false;
 	}
 
 	if (semctl(Semaphore, 0, SETVAL, 0) < 0)
 	{
 		int err = errno;
-		UE_LOG(LogZenServiceInstance, Warning, TEXT("Failed signaling named event '%hs' (err: %d)"), *EventPath, err);
+		UE_LOG(LogZenServiceInstance, Warning, TEXT("Failed signaling named event '%hs' (err: %d)"), *ShutdownEventName, err);
 		return false;
 	}
 	return true;
@@ -1319,200 +1318,229 @@ RequestZenShutdownOnEffectivePort(uint16 EffectiveListenPort)
 #endif
 }
 
-static bool AttemptShutdownUsingExecutablePathOnly(const ZenServerState& ServerState, const TCHAR* ExecutablePath, double MaximumWaitDurationSeconds)
+static bool
+IsZenProcessActive(const TCHAR* ExecutablePath, uint32* OutPid)
 {
-	uint32 ServicePid;
-	if (!IsZenProcessActive(ExecutablePath, &ServicePid))
-	{
-		return true;
-	}
-
-	const ZenServerState::ZenServerEntry* Entry = ServerState.LookupByPid(ServicePid);
-	if (!Entry)
-	{
-		if (!IsZenProcessActive(ExecutablePath, nullptr))
-		{
-			return true;
-		}
-		UE_LOG(LogZenServiceInstance, Warning, TEXT("Can't find server state for running service for executable '%s' (Pid: %u)"), ExecutablePath, ServicePid);
-		return false;
-	}
-
-	uint16 EffectiveListenPort = Entry->EffectiveListenPort.load(std::memory_order_relaxed);
-	if (!RequestZenShutdownOnEffectivePort(EffectiveListenPort))
+	FString NormalizedExecutablePath(ExecutablePath);
+	FPaths::NormalizeFilename(NormalizedExecutablePath);
+	FPlatformProcess::FProcEnumerator ProcIter;
+	while (ProcIter.MoveNext())
 	{
-		if (!IsZenProcessActive(ExecutablePath, &ServicePid))
+		FPlatformProcess::FProcEnumInfo ProcInfo = ProcIter.GetCurrent();
+		FString Candidate = ProcInfo.GetFullPath();
+		FPaths::NormalizeFilename(Candidate);
+		if (Candidate == NormalizedExecutablePath)
 		{
+			if (OutPid)
+			{
+				*OutPid = ProcInfo.GetPID();
+			}
 			return true;
 		}
-		UE_LOG(LogZenServiceInstance, Warning, TEXT("Failed to request shutdown for running service for executable '%s' (Pid: %u, Port: %u)"), ExecutablePath, ServicePid, EffectiveListenPort);
-		return false;
 	}
-
-	uint64 ZenShutdownWaitStartTime = FPlatformTime::Cycles64();
-	while ((Entry != nullptr) || IsZenProcessActive(ExecutablePath, &ServicePid))
-	{
-		double ZenShutdownWaitDuration = FPlatformTime::ToSeconds64(FPlatformTime::Cycles64() - ZenShutdownWaitStartTime);
-		if (ZenShutdownWaitDuration < MaximumWaitDurationSeconds)
-		{
-			FPlatformProcess::Sleep(0.01f);
-		}
-		else
-		{
-			UE_LOG(LogZenServiceInstance, Warning, TEXT("Timed out waiting for shutdown of running service for executable '%s' (Pid: %u)"), ExecutablePath, ServicePid);
-			return false;
-		}
-		Entry = ServerState.LookupByPid(ServicePid);
-	}
-	return true;
+	return false;
 }
 
 static bool
-ShutdownRunningServiceUsingExecutablePath(const TCHAR* ExecutablePath, double MaximumWaitDurationSeconds = 25.0)
+FindZenProcessId(uint32* OutPid)
 {
-	uint32 ServicePid;
-	if (!IsZenProcessActive(ExecutablePath, &ServicePid))
-	{
-		UE_LOG(LogZenServiceInstance, Log, TEXT("No running service found for executable '%s'"), ExecutablePath);
-		return true;
-	}
-
-	UE_LOG(LogZenServiceInstance, Display, TEXT("Waiting for running instance using executable '%s' to shut down"), ExecutablePath);
-
-	const ZenServerState ServerState(/* ReadOnly */true);
-	const ZenServerState::ZenServerEntry* Entry = ServerState.LookupByPid(ServicePid);
-	if (!Entry)
+	FString ExecutableName = GetLocalServiceExecutableName();
+	FPaths::NormalizeFilename(ExecutableName);
+	FPlatformProcess::FProcEnumerator ProcIter;
+	while (ProcIter.MoveNext())
 	{
-		UE_LOG(LogZenServiceInstance, Warning, TEXT("Can't find server state for running service for executable '%s' (Pid: %u), attempting hard terminate"), ExecutablePath, ServicePid);
-		if (NativeTerminate(ServicePid))
+		FPlatformProcess::FProcEnumInfo ProcInfo = ProcIter.GetCurrent();
+		FString Candidate = ProcInfo.GetFullPath();
+		FPaths::NormalizeFilename(Candidate);
+		FString CandidateExecutableName = FPaths::GetPathLeaf(Candidate);
+		if (ExecutableName == CandidateExecutableName)
 		{
+			if (OutPid)
+			{
+				*OutPid = ProcInfo.GetPID();
+			}
 			return true;
 		}
-		UE_LOG(LogZenServiceInstance, Warning, TEXT("Failed to terminate running service for executable '%s' (Pid: %u)"), ExecutablePath, ServicePid);
-		return false;
 	}
+	return false;
+}
 
-	uint16 EffectiveListenPort = Entry->EffectiveListenPort.load(std::memory_order_relaxed);
-
-	if (!RequestZenShutdownOnEffectivePort(EffectiveListenPort))
+static bool ShutdownZenServerProcess(int Pid, double MaximumWaitDurationSeconds = 25.0)
+{
+	const ZenServerState ServerState(/* ReadOnly */true);
+	const ZenServerState::ZenServerEntry* Entry = ServerState.LookupByPid(Pid);
+	if (Entry)
 	{
-		if (!IsZenProcessActive(ExecutablePath, nullptr))
+		uint16 EffectivePort = Entry->EffectiveListenPort.load(std::memory_order_relaxed);
+		UE_LOG(LogZenServiceInstance, Display, TEXT("Requesting shut down of zenserver process %d runnning on effective port %u"), Pid, EffectivePort);
+		if (RequestZenShutdownOnEffectivePort(EffectivePort))
 		{
-			return true;
-		}
-		UE_LOG(LogZenServiceInstance, Warning, TEXT("Failed to request shutdown for running service for executable '%s' (Pid: %u, Port: %u), attempting hard terminate"), ExecutablePath, ServicePid, EffectiveListenPort);
-		if (NativeTerminate(ServicePid))
-		{
-			return true;
+			uint64 ZenShutdownWaitStartTime = FPlatformTime::Cycles64();
+			while (NativeIsProcessRunning(Pid))
+			{
+				double ZenShutdownWaitDuration = FPlatformTime::ToSeconds64(FPlatformTime::Cycles64() - ZenShutdownWaitStartTime);
+				if (ZenShutdownWaitDuration < MaximumWaitDurationSeconds)
+				{
+					FPlatformProcess::Sleep(0.01f);
+				}
+				else
+				{
+					UE_LOG(LogZenServiceInstance, Warning, TEXT("Timed out waiting for shut down of running service with pid %d"), Pid);
+					break;
+				}
+			}
 		}
-		UE_LOG(LogZenServiceInstance, Warning, TEXT("Failed to terminate running service for executable '%s' (Pid: %u, Port: %u)"), ExecutablePath, ServicePid, EffectiveListenPort);
-		return false;
 	}
-
-	uint64 ZenShutdownWaitStartTime = FPlatformTime::Cycles64();
-	while (IsZenProcessActive(ExecutablePath, nullptr))
+	if (NativeIsProcessRunning(Pid))
 	{
-		double ZenShutdownWaitDuration = FPlatformTime::ToSeconds64(FPlatformTime::Cycles64() - ZenShutdownWaitStartTime);
-		if (ZenShutdownWaitDuration < MaximumWaitDurationSeconds)
+		UE_LOG(LogZenServiceInstance, Display, TEXT("Attempting termination of zenserver process with pid %d"), Pid);
+		if (!NativeTerminate(Pid))
 		{
-			FPlatformProcess::Sleep(0.01f);
-		}
-		else
-		{
-			UE_LOG(LogZenServiceInstance, Warning, TEXT("Timed out waiting for shutdown of running service for executable '%s' (Pid: %u, Port: %u)"), ExecutablePath, ServicePid, EffectiveListenPort);
-			return false;
+			if (NativeIsProcessRunning(Pid))
+			{
+				UE_LOG(LogZenServiceInstance, Warning, TEXT("Failed to terminate zenserver process with pid %d"), Pid);
+				return false;
+			}
 		}
 	}
+	UE_LOG(LogZenServiceInstance, Display, TEXT("Successfully shut down zenserver process with pid %d"), Pid);
 	return true;
 }
 
-static bool AttemptShutdownUsingPortOnly(const ZenServerState& ServerState, uint16 EffectiveListenPort, double MaximumWaitDurationSeconds)
+static bool ShutDownZenServerProcessExecutable(const FString& ExecutablePath, double MaximumWaitDurationSeconds = 25.0)
 {
-	if (!RequestZenShutdownOnEffectivePort(EffectiveListenPort))
-	{
-		if (!IsZenProcessUsingEffectivePort(EffectiveListenPort))
-		{
-			return true;
-		}
-		UE_LOG(LogZenServiceInstance, Warning, TEXT("Failed to request shutdown for running service using port %u"), EffectiveListenPort);
-		return false;
-	}
-
-	const ZenServerState::ZenServerEntry* Entry = ServerState.LookupByEffectiveListenPort(EffectiveListenPort);
 	uint64 ZenShutdownWaitStartTime = FPlatformTime::Cycles64();
-	while ((Entry != nullptr) || IsZenProcessUsingEffectivePort(EffectiveListenPort))
+	uint32_t Pid = 0;
+	while (IsZenProcessActive(*ExecutablePath, &Pid))
 	{
+		UE_LOG(LogZenServiceInstance, Display, TEXT("Attempting to shut down of zenserver executable '%s' process with pid %d"), *ExecutablePath, Pid);
 		double ZenShutdownWaitDuration = FPlatformTime::ToSeconds64(FPlatformTime::Cycles64() - ZenShutdownWaitStartTime);
-		if (ZenShutdownWaitDuration < MaximumWaitDurationSeconds)
+		if (ShutdownZenServerProcess(Pid, ZenShutdownWaitDuration))
 		{
-			FPlatformProcess::Sleep(0.01f);
+			return true;
 		}
 		else
 		{
-			UE_LOG(LogZenServiceInstance, Warning, TEXT("Timed out waiting for shutdown of running service using port %u"), EffectiveListenPort);
+			UE_LOG(LogZenServiceInstance, Warning, TEXT("Failed to shut down zenserver executable '%s' process with pid %d"), *ExecutablePath, Pid);
 			return false;
 		}
-		Entry = ServerState.LookupByEffectiveListenPort(EffectiveListenPort);
 	}
-	// Sleep one second to ensure the process can run the final exit code and not hindering overwriting the executable
-	FPlatformProcess::Sleep(1.0f);
 	return true;
 }
 
-static bool
-ShutdownRunningServiceUsingEffectivePort(uint16 EffectiveListenPort, double MaximumWaitDurationSeconds = 25.0)
+static bool ShutDownZenServerProcessLockingDataDir(const FString& DataPath, double MaximumWaitDurationSeconds = 25.0)
 {
-	if (!IsZenProcessUsingEffectivePort(EffectiveListenPort))
+	const FString LockFilePath = FPaths::Combine(DataPath, TEXT(".lock"));
+
+	uint64 ZenShutdownWaitStartTime = FPlatformTime::Cycles64();
+	if (!IsLockFileLocked(*LockFilePath, true))
 	{
+		UE_LOG(LogZenServiceInstance, Display, TEXT("Lock file '%s' is not active, nothing to do"), *LockFilePath);
 		return true;
 	}
-
-	UE_LOG(LogZenServiceInstance, Display, TEXT("Waiting for running instance using port %u to shut down"), EffectiveListenPort);
-
-	const ZenServerState ServerState(/*ReadOnly*/true);
-	const ZenServerState::ZenServerEntry* Entry = ServerState.LookupByEffectiveListenPort(EffectiveListenPort);
-	if (!Entry)
+	LockFileData LockFileState = ReadCbLockFile(LockFilePath);
+	if (!LockFileState.IsValid)
 	{
-		if (!IsZenProcessUsingEffectivePort(EffectiveListenPort))
+		while (true)
 		{
+			if (!IsLockFileLocked(*LockFilePath, true))
+			{
+				return true;
+			}
+			uint32_t Pid = 0;
+			if (!FindZenProcessId(&Pid))
+			{
+				UE_LOG(LogZenServiceInstance, Warning, TEXT("Failed to find zenserver process locking file '%s'"), *LockFilePath);
+				return false;
+			}
+			UE_LOG(LogZenServiceInstance, Warning, TEXT("Found locked but invalid lock file at '%s', attempting shut down of zenserver process with pid %d"), *LockFilePath, Pid);
+			double ZenShutdownWaitDuration = FPlatformTime::ToSeconds64(FPlatformTime::Cycles64() - ZenShutdownWaitStartTime);
+			if (!ShutdownZenServerProcess(Pid, ZenShutdownWaitDuration))
+			{
+				break;
+			}
+		}
+		if (!IsLockFileLocked(*LockFilePath))
+		{
+			UE_LOG(LogZenServiceInstance, Display, TEXT("Successfully shut down zenserver using lock file '%s'"), *LockFilePath);
 			return true;
 		}
-
-		UE_LOG(LogZenServiceInstance, Warning, TEXT("Can't find server state for running service using port %u, attempting blind shut down"), EffectiveListenPort);
-		return AttemptShutdownUsingPortOnly(ServerState, EffectiveListenPort, MaximumWaitDurationSeconds);
+		UE_LOG(LogZenServiceInstance, Warning, TEXT("Failed to shut down zenserver process locking file '%s'"), *LockFilePath);
+		return false;
 	}
 
-	uint32 ServicePid = Entry->Pid.load(std::memory_order_relaxed);
-	if (!RequestZenShutdownOnEffectivePort(EffectiveListenPort))
+	uint16 EffectivePort = LockFileState.EffectivePort;
+
+	const ZenServerState ServerState(/* ReadOnly */true);
+	const ZenServerState::ZenServerEntry* Entry = ServerState.LookupByEffectiveListenPort(EffectivePort);
+	if (Entry)
 	{
-		if (!NativeIsProcessRunning(ServicePid) && !IsZenProcessUsingEffectivePort(EffectiveListenPort))
+		UE_LOG(LogZenServiceInstance, Display, TEXT("Requesting shut down of zenserver process using lock file '%s' with effective port %d"), *LockFilePath, EffectivePort);
+		if (RequestZenShutdownOnEffectivePort(EffectivePort))
 		{
-			return true;
+			while (IsLockFileLocked(*LockFilePath, true))
+			{
+				double ZenShutdownWaitDuration = FPlatformTime::ToSeconds64(FPlatformTime::Cycles64() - ZenShutdownWaitStartTime);
+				if (ZenShutdownWaitDuration < MaximumWaitDurationSeconds)
+				{
+					FPlatformProcess::Sleep(0.01f);
+				}
+				else
+				{
+					UE_LOG(LogZenServiceInstance, Warning, TEXT("Timed out waiting for shut down of zensever process using lock file '%s' with effective port %u"), *LockFilePath, EffectivePort);
+					break;
+				}
+			}
+			if (!IsLockFileLocked(*LockFilePath, true))
+			{
+				UE_LOG(LogZenServiceInstance, Display, TEXT("Successfully shut down zenserver process using lock file '%s' with effective port %u"), *LockFilePath, EffectivePort);
+				return true;
+			}
 		}
-		UE_LOG(LogZenServiceInstance, Warning, TEXT("Failed to request shutdown for running service using port %u (Pid: %u)"), EffectiveListenPort, ServicePid);
-		return false;
 	}
 
-	uint64 ZenShutdownWaitStartTime = FPlatformTime::Cycles64();
-	while (NativeIsProcessRunning(ServicePid) || IsZenProcessUsingEffectivePort(EffectiveListenPort))
+	while (true)
 	{
-		double ZenShutdownWaitDuration = FPlatformTime::ToSeconds64(FPlatformTime::Cycles64() - ZenShutdownWaitStartTime);
-		if (ZenShutdownWaitDuration < MaximumWaitDurationSeconds)
+		if (!IsLockFileLocked(*LockFilePath, true))
 		{
-			FPlatformProcess::Sleep(0.01f);
+			return true;
 		}
-		else
+		uint32_t Pid = 0;
+		if (!FindZenProcessId(&Pid))
 		{
-			UE_LOG(LogZenServiceInstance, Warning, TEXT("Timed out waiting for shutdown of running service using port %u (Pid: %u)"), EffectiveListenPort, ServicePid);
+			UE_LOG(LogZenServiceInstance, Warning, TEXT("Failed to find zenserver process locking file '%s'"), *LockFilePath);
 			return false;
 		}
-		Entry = ServerState.LookupByEffectiveListenPort(EffectiveListenPort);
-		if (Entry)
+		UE_LOG(LogZenServiceInstance, Warning, TEXT("Found locked but invalid lock file at '%s', attempting shut down of zenserver process with pid %d"), *LockFilePath, Pid);
+		double ZenShutdownWaitDuration = FPlatformTime::ToSeconds64(FPlatformTime::Cycles64() - ZenShutdownWaitStartTime);
+		if (!ShutdownZenServerProcess(Pid, ZenShutdownWaitDuration))
 		{
-			ServicePid = Entry->Pid.load(std::memory_order_relaxed);
+			break;
 		}
 	}
-	return true;
+
+	if (!IsLockFileLocked(*LockFilePath))
+	{
+		UE_LOG(LogZenServiceInstance, Display, TEXT("Successfully shut down zenserver using lock file '%s'"), *LockFilePath);
+		return true;
+	}
+	UE_LOG(LogZenServiceInstance, Warning, TEXT("Failed to shut down zenserver process locking file '%s'"), *LockFilePath);
+	return false;
+}
+
+static bool
+IsZenProcessUsingDataDir(const TCHAR* LockFilePath, LockFileData* OutLockFileData)
+{
+	if (IsLockFileLocked(LockFilePath, true))
+	{
+		if (OutLockFileData)
+		{
+			// If an instance is running with this data path, check if we can use it and what port it is on
+			*OutLockFileData = ReadCbLockFile(LockFilePath);
+		}
+		return true;
+	}
+	return false;
 }
 
 static FString
@@ -1744,29 +1772,13 @@ StopLocalService(const TCHAR* DataPath, double MaximumWaitDurationSeconds)
 {
 	const FString LockFilePath = FPaths::Combine(DataPath, TEXT(".lock"));
 	LockFileData LockFileState;
-	if (IsZenProcessUsingDataDir(*LockFilePath, &LockFileState))
+	if (IsLockFileLocked(*LockFilePath, true))
 	{
-		if (LockFileState.IsValid)
-		{
-			return ShutdownRunningServiceUsingEffectivePort(LockFileState.EffectivePort, MaximumWaitDurationSeconds);
-		}
-		return false;
+		return ShutDownZenServerProcessLockingDataDir(DataPath, MaximumWaitDurationSeconds);
 	}
 	return true;
 }
 
-FString
-GetLocalServiceInstallPath()
-{
-	return FPaths::ConvertRelativePathToFull(FPaths::Combine(FPlatformProcess::ApplicationSettingsDir(), TEXT("Zen\\Install"),
-#if PLATFORM_WINDOWS
-		TEXT("zenserver.exe")
-#else
-		TEXT("zenserver")
-#endif
-		));
-}
-
 FString
 GetLocalServiceInstallVersion(bool bDetailed)
 {
@@ -1960,7 +1972,7 @@ FZenServiceInstance::TryRecovery()
 					}
 				}
 			}
-			if (bShutdownExistingInstance && !ShutdownRunningServiceUsingEffectivePort(Port))
+			if (bShutdownExistingInstance && !ShutdownZenServerProcess((int)AutoLaunchedPid))	// !ShutdownRunningServiceUsingEffectivePort(Port))
 			{
 				return false;
 			}
@@ -2124,7 +2136,6 @@ PromptUserOfFailedShutDownOfExistingProcess(uint16 Port)
 	}
 }
 
-
 FString
 FZenServiceInstance::ConditionalUpdateLocalInstall()
 {
@@ -2141,7 +2152,7 @@ FZenServiceInstance::ConditionalUpdateLocalInstall()
 	if (IsInstallVersionOutOfDate(InTreeUtilityPath, InstallUtilityPath, InTreeServicePath, InstallServicePath, InTreeVersionCache, InstallVersionCache))
 	{
 		UE_LOG(LogZenServiceInstance, Display, TEXT("Installing service from '%s' to '%s'"), *InTreeServicePath, *InstallServicePath);
-		if (!ShutdownRunningServiceUsingExecutablePath(*InstallServicePath))
+		if (!ShutDownZenServerProcessExecutable(InstallServicePath))
 		{
 			PromptUserToStopRunningServerInstanceForUpdate(InstallServicePath);
 			return FString();
@@ -2209,7 +2220,7 @@ FZenServiceInstance::AutoLaunch(const FServiceAutoLaunchSettings& InSettings, FS
 		LockFileState = LockFileData();
 	}
 
-	uint16 ShutdownEffectivePort = 0;
+	bool bShutDownExistningInstance = true;
 	bool bLaunchNewInstance = true;
 
 	if (LockFileState.IsReady)
@@ -2217,9 +2228,13 @@ FZenServiceInstance::AutoLaunch(const FServiceAutoLaunchSettings& InSettings, FS
 		const ZenServerState State(/*ReadOnly*/true);
 		if (State.LookupByPid(LockFileState.ProcessId) == nullptr)
 		{
-			PromptUserOfLockedDataFolder(*InSettings.DataPath);
-			FPlatformMisc::RequestExit(true);
-			return false;
+			if (IsZenProcessUsingDataDir(*LockFilePath, nullptr))
+			{
+				UE_LOG(LogZenServiceInstance, Warning, TEXT("Found locked valid lock file '%s' but no matching process (Pid: %d), exiting"), *LockFilePath, LockFileState.ProcessId);
+				PromptUserOfLockedDataFolder(*InSettings.DataPath);
+				FPlatformMisc::RequestExit(true);
+				return false;
+			}
 		}
 
 		if (InSettings.bIsDefaultSharedRunContext)
@@ -2232,46 +2247,53 @@ FZenServiceInstance::AutoLaunch(const FServiceAutoLaunchSettings& InSettings, FS
 			DesiredRunContext.bShowConsole = InSettings.bShowConsole;
 
 			FZenLocalServiceRunContext CurrentRunContext;
-			if (CurrentRunContext.ReadFromJsonFile(*ExecutionContextFilePath) && (DesiredRunContext == CurrentRunContext))
+
+			bool ReadCurrentContextOK = CurrentRunContext.ReadFromJsonFile(*ExecutionContextFilePath);
+			if (ReadCurrentContextOK && (DesiredRunContext == CurrentRunContext))
 			{
-				UE_LOG(LogZenServiceInstance, Log, TEXT("Found existing instance running on port %u matching our settings"), InSettings.DesiredPort);
+				UE_LOG(LogZenServiceInstance, Log, TEXT("Found existing instance running on port %u matching our settings, no actions needed"), InSettings.DesiredPort);
 				bLaunchNewInstance = false;
+				bShutDownExistningInstance = false;
 			}
 			else
 			{
-				UE_LOG(LogZenServiceInstance, Log, TEXT("Found existing instance running on port %u with different run context, will attempt shut down"), InSettings.DesiredPort);
-				ShutdownEffectivePort = LockFileState.EffectivePort;
+				FString JsonTcharText;
+				{
+					TSharedRef<TJsonWriter<TCHAR, TPrettyJsonPrintPolicy<TCHAR>>> Writer = TJsonWriterFactory<TCHAR, TPrettyJsonPrintPolicy<TCHAR>>::Create(&JsonTcharText);
+					Writer->WriteObjectStart();
+					Writer->WriteObjectStart("Current");
+					CurrentRunContext.WriteToJson(*Writer);
+					Writer->WriteObjectEnd();
+					Writer->WriteObjectStart("Desired");
+					DesiredRunContext.WriteToJson(*Writer);
+					Writer->WriteObjectEnd();
+					Writer->WriteObjectEnd();
+					Writer->Close();
+				}
+				UE_LOG(LogZenServiceInstance, Log, TEXT("Found existing instance running on port %u with different run context, will attempt shut down\n{%s}"), InSettings.DesiredPort, *JsonTcharText);
+				bShutDownExistningInstance = true;
+				bLaunchNewInstance = true;
 			}
 		}
 		else
 		{
 			UE_LOG(LogZenServiceInstance, Log, TEXT("Found existing instance running on port %u when not using shared context, will use it"), InSettings.DesiredPort);
+			bShutDownExistningInstance = false;
+			bLaunchNewInstance = false;
 		}
 	}
 	else
 	{
-		const ZenServerState State(/*ReadOnly*/true);
-		const ZenServerState::ZenServerEntry* RunningEntry = State.LookupByDesiredListenPort(InSettings.DesiredPort);
-		if (RunningEntry != nullptr)
-		{
-			ShutdownEffectivePort = RunningEntry->EffectiveListenPort.load(std::memory_order_relaxed);
-			UE_LOG(LogZenServiceInstance, Log, TEXT("Found existing instance running on port %u with different data directory, will attempt shut down"), ShutdownEffectivePort);
-		}
+		UE_LOG(LogZenServiceInstance, Log, TEXT("No current process using the data dir found, launching a new instance"));
+		bShutDownExistningInstance = false;
+		bLaunchNewInstance = true;
 	}
 
-	if (ShutdownEffectivePort != 0)
+	if (bShutDownExistningInstance)
 	{
-		if (!ShutdownRunningServiceUsingEffectivePort(ShutdownEffectivePort))
-		{
-			PromptUserOfFailedShutDownOfExistingProcess(ShutdownEffectivePort);
-			FPlatformMisc::RequestExit(true);
-			return false;
-		}
-		checkf(!IsZenProcessUsingEffectivePort(ShutdownEffectivePort), TEXT("Service port is still in use %u"), ShutdownEffectivePort);
-
-		if (IsLockFileLocked(*LockFilePath))
+		if (!ShutDownZenServerProcessLockingDataDir(InSettings.DataPath))
 		{
-			PromptUserOfLockedDataFolder(*InSettings.DataPath);
+			PromptUserOfFailedShutDownOfExistingProcess(InSettings.DesiredPort);
 			FPlatformMisc::RequestExit(true);
 			return false;
 		}
