Zenserver autolaunch and shutdown unit tests.

diff --git a/Engine/Source/Developer/Zen/Tests/ZenTest.cpp b/Engine/Source/Developer/Zen/Tests/ZenTest.cpp
new file mode 100644
--- /dev/null
+++ b/Engine/Source/Developer/Zen/Tests/ZenTest.cpp
@@ -0,0 +1,105 @@
+// Copyright Epic Games, Inc. All Rights Reserved.
+
+#if WITH_LOW_LEVEL_TESTS
+
+
+#include "TestHarness.h"
+
+#if WITH_EDITORONLY_DATA
+#include "Experimental/ZenServerInterface.h"
+#include "HAL/FileManager.h"
+#include "Misc/Paths.h"
+
+namespace UE
+{
+
+TEST_CASE("Zen::ZenServerInterface", "[Zen][Basic]")
+{
+	using namespace UE::Zen;
+	uint16 DefaultTestPort = 8559;
+	uint16 CurrentTestPort = DefaultTestPort + 1;
+	FString DefaultArgs = TEXT("--http asio");
+	FString DataPathRoot = FPaths::ConvertRelativePathToFull(FPaths::Combine(FPaths::EngineSavedDir(), "ZenServerInterfaceUnitTest"));
+	FString DefaultDataPath = FPaths::ConvertRelativePathToFull(FPaths::Combine(DataPathRoot, "Default"));
+	IFileManager::Get().DeleteDirectory(*DataPathRoot, false, true);
+
+	SECTION("Basic AutoLaunch and Shutdown")
+	{
+		for (int Iteration = 0; Iteration < 2; ++Iteration)
+		{
+			FServiceSettings ZenTestServiceSettings;
+			FServiceAutoLaunchSettings& ZenTestAutoLaunchSettings = ZenTestServiceSettings.SettingsVariant.Get<FServiceAutoLaunchSettings>();
+			ZenTestAutoLaunchSettings.DataPath = DefaultDataPath;
+			ZenTestAutoLaunchSettings.ExtraArgs = DefaultArgs;
+			ZenTestAutoLaunchSettings.DesiredPort = DefaultTestPort;
+
+			{
+				FScopeZenService ScopeZenService(MoveTemp(ZenTestServiceSettings));
+				FZenServiceInstance& ZenInstance = ScopeZenService.GetInstance();
+				uint16 AutoLaunchedPort = ZenInstance.GetAutoLaunchedPort();
+				uint16 DetectedPort = 0;
+				CHECK(ZenInstance.IsServiceReady());
+
+				CHECK(IsLocalServiceRunning(*DefaultDataPath, &DetectedPort));
+				CHECK(DetectedPort == AutoLaunchedPort);
+				CHECK(StopLocalService(*DefaultDataPath));
+				CHECK(!IsLocalServiceRunning(*DefaultDataPath));
+			}
+		}
+	}
+
+	SECTION("Overlapping AutoLaunch and Shutdown")
+	{
+		for (int Iteration = 0; Iteration < 3; ++Iteration)
+		{
+			FServiceSettings ZenTestServiceSettings;
+			FServiceAutoLaunchSettings& ZenTestAutoLaunchSettings = ZenTestServiceSettings.SettingsVariant.Get<FServiceAutoLaunchSettings>();
+			ZenTestAutoLaunchSettings.DataPath = DefaultDataPath;
+			ZenTestAutoLaunchSettings.ExtraArgs = DefaultArgs;
+			ZenTestAutoLaunchSettings.DesiredPort = DefaultTestPort;
+
+			{
+				FScopeZenService ScopeZenService(MoveTemp(ZenTestServiceSettings));
+				FZenServiceInstance& ZenInstance = ScopeZenService.GetInstance();
+				uint16 AutoLaunchedPort = ZenInstance.GetAutoLaunchedPort();
+				uint16 DetectedPort = 0;
+				CHECK(ZenInstance.IsServiceReady());
+
+				CHECK(IsLocalServiceRunning(*DefaultDataPath, &DetectedPort));
+				CHECK(DetectedPort == AutoLaunchedPort);
+			}
+		}
+		CHECK(StopLocalService(*DefaultDataPath));
+		CHECK(!IsLocalServiceRunning(*DefaultDataPath));
+	}
+
+	SECTION("Overlapping AutoLaunch and Shutdown With DataPath Shared And Differing Args")
+	{
+		for (int Iteration = 0; Iteration < 3; ++Iteration)
+		{
+			FServiceSettings ZenTestServiceSettings;
+			FServiceAutoLaunchSettings& ZenTestAutoLaunchSettings = ZenTestServiceSettings.SettingsVariant.Get<FServiceAutoLaunchSettings>();
+			ZenTestAutoLaunchSettings.DataPath = DefaultDataPath;
+			ZenTestAutoLaunchSettings.ExtraArgs = *WriteToString<128>(DefaultArgs, TEXT(" --gc-interval-seconds "), (Iteration+1)*1000);
+			ZenTestAutoLaunchSettings.DesiredPort = DefaultTestPort;
+
+			{
+				FScopeZenService ScopeZenService(MoveTemp(ZenTestServiceSettings));
+				FZenServiceInstance& ZenInstance = ScopeZenService.GetInstance();
+				uint16 AutoLaunchedPort = ZenInstance.GetAutoLaunchedPort();
+				uint16 DetectedPort = 0;
+				CHECK(ZenInstance.IsServiceReady());
+
+				CHECK(IsLocalServiceRunning(*DefaultDataPath, &DetectedPort));
+				CHECK(DetectedPort == AutoLaunchedPort);
+			}
+		}
+		CHECK(StopLocalService(*DefaultDataPath));
+		CHECK(!IsLocalServiceRunning(*DefaultDataPath));
+	}
+}
+
+} // UE
+#endif // WITH_EDITORONLY_DATA
+
+#endif // WITH_LOW_LEVEL_TESTS
