Fix issue where zenserver fails to launch when DDC path differs between two workspaces.

diff --git a/Engine/Source/Developer/Zen/Private/ZenServerInterface.cpp b/Engine/Source/Developer/Zen/Private/ZenServerInterface.cpp
--- a/Engine/Source/Developer/Zen/Private/ZenServerInterface.cpp
+++ b/Engine/Source/Developer/Zen/Private/ZenServerInterface.cpp
@@ -2220,7 +2220,8 @@ FZenServiceInstance::AutoLaunch(const FServiceAutoLaunchSettings& InSettings, FS
 		LockFileState = LockFileData();
 	}
 
-	bool bShutDownExistningInstance = true;
+	bool bShutDownExistingInstanceForDataPath = true;
+	uint32 ShutdownExistingInstanceForPid = 0;
 	bool bLaunchNewInstance = true;
 
 	if (LockFileState.IsReady)
@@ -2253,7 +2254,7 @@ FZenServiceInstance::AutoLaunch(const FServiceAutoLaunchSettings& InSettings, FS
 			{
 				UE_LOG(LogZenServiceInstance, Log, TEXT("Found existing instance running on port %u matching our settings, no actions needed"), InSettings.DesiredPort);
 				bLaunchNewInstance = false;
-				bShutDownExistningInstance = false;
+				bShutDownExistingInstanceForDataPath = false;
 			}
 			else
 			{
@@ -2271,25 +2272,39 @@ FZenServiceInstance::AutoLaunch(const FServiceAutoLaunchSettings& InSettings, FS
 					Writer->Close();
 				}
 				UE_LOG(LogZenServiceInstance, Log, TEXT("Found existing instance running on port %u with different run context, will attempt shut down\n{%s}"), InSettings.DesiredPort, *JsonTcharText);
-				bShutDownExistningInstance = true;
+				bShutDownExistingInstanceForDataPath = true;
 				bLaunchNewInstance = true;
 			}
 		}
 		else
 		{
 			UE_LOG(LogZenServiceInstance, Log, TEXT("Found existing instance running on port %u when not using shared context, will use it"), InSettings.DesiredPort);
-			bShutDownExistningInstance = false;
+			bShutDownExistingInstanceForDataPath = false;
 			bLaunchNewInstance = false;
 		}
 	}
 	else
 	{
-		UE_LOG(LogZenServiceInstance, Log, TEXT("No current process using the data dir found, launching a new instance"));
-		bShutDownExistningInstance = false;
+		const ZenServerState State(/*ReadOnly*/true);
+		const ZenServerState::ZenServerEntry* RunningEntry = State.LookupByDesiredListenPort(InSettings.DesiredPort);
+		if (RunningEntry != nullptr)
+		{
+			// It is necessary to tear down an existing zenserver running on our desired port but in a different data path because:
+			// 1. zenserver won't accept port collision with itself, and will instead say "Exiting since there is already a process listening to port ..."
+			// 2. When UE is changing data directories (eg: DDC path config change) we don't want to leave zenservers running on the past directories for no reason
+			// Unlike other shutdown scenarios, this one can't be done based on our desired data path because the zenserver we want to shut down is running in a different data path
+			UE_LOG(LogZenServiceInstance, Log, TEXT("Found existing instance running on port %u with different data directory, will attempt shutdown"), InSettings.DesiredPort);
+			ShutdownExistingInstanceForPid = RunningEntry->Pid;
+		}
+		else
+		{
+			UE_LOG(LogZenServiceInstance, Log, TEXT("No current process using the data dir found, launching a new instance"));
+		}
+		bShutDownExistingInstanceForDataPath = false;
 		bLaunchNewInstance = true;
 	}
 
-	if (bShutDownExistningInstance)
+	if (bShutDownExistingInstanceForDataPath)
 	{
 		if (!ShutDownZenServerProcessLockingDataDir(InSettings.DataPath))
 		{
@@ -2299,6 +2314,16 @@ FZenServiceInstance::AutoLaunch(const FServiceAutoLaunchSettings& InSettings, FS
 		}
 	}
 
+	if (ShutdownExistingInstanceForPid != 0)
+	{
+		if (!ShutdownZenServerProcess(ShutdownExistingInstanceForPid))
+		{
+			PromptUserOfFailedShutDownOfExistingProcess(InSettings.DesiredPort);
+			FPlatformMisc::RequestExit(true);
+			return false;
+		}
+	}
+
 	// When limiting process lifetime, always re-launch to add sponsor process IDs.
 	// When not limiting process lifetime, only launch if the process is not already live.
 	if (bLaunchNewInstance)
diff --git a/Engine/Source/Developer/Zen/Tests/ZenTest.cpp b/Engine/Source/Developer/Zen/Tests/ZenTest.cpp
--- a/Engine/Source/Developer/Zen/Tests/ZenTest.cpp
+++ b/Engine/Source/Developer/Zen/Tests/ZenTest.cpp
@@ -73,7 +73,47 @@ TEST_CASE("Zen::ZenServerInterface", "[Zen][Basic]")
 		CHECK(!IsLocalServiceRunning(*DefaultDataPath));
 	}
 
-	SECTION("Overlapping AutoLaunch and Shutdown With DataPath Shared And Differing Args")
+	SECTION("Server moving between separate datapaths with same port")
+	{
+		FString LastDataPath;
+		{
+			TArray<FScopeZenService> ConcurrentServices;
+			for (int Iteration = 0; Iteration < 3; ++Iteration)
+			{
+				FString DataPath = FPaths::ConvertRelativePathToFull(FPaths::Combine(DataPathRoot, FString::Printf(TEXT("Instance%d"), Iteration)));
+				LastDataPath = DataPath;
+				FServiceSettings ZenTestServiceSettings;
+				FServiceAutoLaunchSettings& ZenTestAutoLaunchSettings = ZenTestServiceSettings.SettingsVariant.Get<FServiceAutoLaunchSettings>();
+				ZenTestAutoLaunchSettings.DataPath = DataPath;
+				ZenTestAutoLaunchSettings.ExtraArgs = DefaultArgs;
+				ZenTestAutoLaunchSettings.DesiredPort = DefaultTestPort;
+
+				{
+					FScopeZenService& ScopeZenService = ConcurrentServices.Emplace_GetRef(MoveTemp(ZenTestServiceSettings));
+					FZenServiceInstance& ZenInstance = ScopeZenService.GetInstance();
+					uint16 AutoLaunchedPort = ZenInstance.GetAutoLaunchedPort();
+					uint16 DetectedPort = 0;
+					CHECK(ZenInstance.IsServiceReady());
+
+					CHECK(IsLocalServiceRunning(*DataPath, &DetectedPort));
+					CHECK(DetectedPort == AutoLaunchedPort);
+
+					for (int PastIteration = 0; PastIteration < Iteration; ++PastIteration)
+					{
+						// By starting a new service on the same port as an existing one, we expect that the existing one
+						// must be torn down.  Check that here.
+						const FServiceAutoLaunchSettings& PastZenTestAutoLaunchSettings = ConcurrentServices[PastIteration].GetInstance()
+							.GetServiceSettings().SettingsVariant.Get<FServiceAutoLaunchSettings>();
+						CHECK(!IsLocalServiceRunning(*PastZenTestAutoLaunchSettings.DataPath));
+					}
+				}
+			}
+		}
+		CHECK(StopLocalService(*LastDataPath));
+		CHECK(!IsLocalServiceRunning(*LastDataPath));
+	}
+
+	SECTION("Overlapping AutoLaunch and Shutdown with DataPath shared and differing args")
 	{
 		for (int Iteration = 0; Iteration < 3; ++Iteration)
 		{
