Partial backout (test only) while failure on Linux is investigated.
[Backout] - CL31365105
Original CL Desc
-----------------------------------------------------------------
Fix issue where zenserver fails to launch when DDC path differs between two workspaces.

diff --git a/Engine/Source/Developer/Zen/Tests/ZenTest.cpp b/Engine/Source/Developer/Zen/Tests/ZenTest.cpp
--- a/Engine/Source/Developer/Zen/Tests/ZenTest.cpp
+++ b/Engine/Source/Developer/Zen/Tests/ZenTest.cpp
@@ -73,47 +73,7 @@ TEST_CASE("Zen::ZenServerInterface", "[Zen][Basic]")
 		CHECK(!IsLocalServiceRunning(*DefaultDataPath));
 	}
 
-	SECTION("Server moving between separate datapaths with same port")
-	{
-		FString LastDataPath;
-		{
-			TArray<FScopeZenService> ConcurrentServices;
-			for (int Iteration = 0; Iteration < 3; ++Iteration)
-			{
-				FString DataPath = FPaths::ConvertRelativePathToFull(FPaths::Combine(DataPathRoot, FString::Printf(TEXT("Instance%d"), Iteration)));
-				LastDataPath = DataPath;
-				FServiceSettings ZenTestServiceSettings;
-				FServiceAutoLaunchSettings& ZenTestAutoLaunchSettings = ZenTestServiceSettings.SettingsVariant.Get<FServiceAutoLaunchSettings>();
-				ZenTestAutoLaunchSettings.DataPath = DataPath;
-				ZenTestAutoLaunchSettings.ExtraArgs = DefaultArgs;
-				ZenTestAutoLaunchSettings.DesiredPort = DefaultTestPort;
-
-				{
-					FScopeZenService& ScopeZenService = ConcurrentServices.Emplace_GetRef(MoveTemp(ZenTestServiceSettings));
-					FZenServiceInstance& ZenInstance = ScopeZenService.GetInstance();
-					uint16 AutoLaunchedPort = ZenInstance.GetAutoLaunchedPort();
-					uint16 DetectedPort = 0;
-					CHECK(ZenInstance.IsServiceReady());
-
-					CHECK(IsLocalServiceRunning(*DataPath, &DetectedPort));
-					CHECK(DetectedPort == AutoLaunchedPort);
-
-					for (int PastIteration = 0; PastIteration < Iteration; ++PastIteration)
-					{
-						// By starting a new service on the same port as an existing one, we expect that the existing one
-						// must be torn down.  Check that here.
-						const FServiceAutoLaunchSettings& PastZenTestAutoLaunchSettings = ConcurrentServices[PastIteration].GetInstance()
-							.GetServiceSettings().SettingsVariant.Get<FServiceAutoLaunchSettings>();
-						CHECK(!IsLocalServiceRunning(*PastZenTestAutoLaunchSettings.DataPath));
-					}
-				}
-			}
-		}
-		CHECK(StopLocalService(*LastDataPath));
-		CHECK(!IsLocalServiceRunning(*LastDataPath));
-	}
-
-	SECTION("Overlapping AutoLaunch and Shutdown with DataPath shared and differing args")
+	SECTION("Overlapping AutoLaunch and Shutdown With DataPath Shared And Differing Args")
 	{
 		for (int Iteration = 0; Iteration < 3; ++Iteration)
 		{
