Fix a bug where ZenServer data path would not update if you initiated an editor restart after changing the Global Local Data Cache Path.  This was due to environment variable inheritance of UE-ZenSubprocessDataPath.  This change ensures that environment variable is cleared as part of shutdown before the editor spawns a new instance of itself.  Also fixed some whitespace issues in passing.

diff --git a/Engine/Source/Developer/Zen/Private/ZenServerInterface.cpp b/Engine/Source/Developer/Zen/Private/ZenServerInterface.cpp
--- a/Engine/Source/Developer/Zen/Private/ZenServerInterface.cpp
+++ b/Engine/Source/Developer/Zen/Private/ZenServerInterface.cpp
@@ -611,9 +611,9 @@ static bool
 DetermineDataPath(const TCHAR* ConfigSection, FString& DataPath, bool& HasInvalidPathConfigurations)
 {
 	auto ValidateDataPath = [](const FString& InDataPath)
-{
-		if (InDataPath.IsEmpty())
 	{
+		if (InDataPath.IsEmpty())
+		{
 			return FString{};
 		}
 		IFileManager& FileManager = IFileManager::Get();
@@ -631,7 +631,7 @@ DetermineDataPath(const TCHAR* ConfigSection, FString& DataPath, bool& HasInvali
 			TestFile->Close();
 			delete TestFile;
 			FileManager.Delete(*TestFilePath);
-		return FinalPath;
+			return FinalPath;
 		}
 		if (FileManager.MakeDirectory(*InDataPath, true))
 		{
@@ -645,9 +645,9 @@ DetermineDataPath(const TCHAR* ConfigSection, FString& DataPath, bool& HasInvali
 	if (FParse::Value(FCommandLine::Get(), TEXT("ZenDataPath="), CommandLineOverrideValue) && !CommandLineOverrideValue.IsEmpty())
 	{
 		if (FString Path = ValidateDataPath(CommandLineOverrideValue); !Path.IsEmpty())
-	{
+		{
 			DataPath = Path;
-		UE_LOG(LogZenServiceInstance, Log, TEXT("Found command line override ZenDataPath=%s"), *CommandLineOverrideValue);
+			UE_LOG(LogZenServiceInstance, Log, TEXT("Found command line override ZenDataPath=%s"), *CommandLineOverrideValue);
 			return true;
 		}
 		UE_LOG(LogZenServiceInstance, Warning, TEXT("Skipping command line override ZenDataPath=%s due to an invalid path"), *CommandLineOverrideValue);
@@ -658,9 +658,9 @@ DetermineDataPath(const TCHAR* ConfigSection, FString& DataPath, bool& HasInvali
 	if (FString SubprocessDataPathEnvOverrideValue = FPlatformMisc::GetEnvironmentVariable(TEXT("UE-ZenSubprocessDataPath")); !SubprocessDataPathEnvOverrideValue.IsEmpty())
 	{
 		if (FString Path = ValidateDataPath(SubprocessDataPathEnvOverrideValue); !Path.IsEmpty())
-	{
+		{
 			DataPath = Path;
-		UE_LOG(LogZenServiceInstance, Log, TEXT("Found subprocess environment variable UE-ZenSubprocessDataPath=%s"), *SubprocessDataPathEnvOverrideValue);
+			UE_LOG(LogZenServiceInstance, Log, TEXT("Found subprocess environment variable UE-ZenSubprocessDataPath=%s"), *SubprocessDataPathEnvOverrideValue);
 			return true;
 		}
 		UE_LOG(LogZenServiceInstance, Warning, TEXT("Skipping subprocess environment variable UE-ZenSubprocessDataPath=%s due to an invalid path"), *SubprocessDataPathEnvOverrideValue);
@@ -685,7 +685,7 @@ DetermineDataPath(const TCHAR* ConfigSection, FString& DataPath, bool& HasInvali
 	if (FString ZenDataPathEnvOverrideValue = FPlatformMisc::GetEnvironmentVariable(TEXT("UE-ZenDataPath")); !ZenDataPathEnvOverrideValue.IsEmpty())
 	{
 		if (FString Path = ValidateDataPath(ZenDataPathEnvOverrideValue); !Path.IsEmpty())
-	{
+		{
 			DataPath = Path;
 			UE_LOG(LogZenServiceInstance, Log, TEXT("Found environment variable UE-ZenDataPath=%s"), *ZenDataPathEnvOverrideValue);
 			return true;
diff --git a/Engine/Source/Editor/UnrealEd/Private/UnrealEdMisc.cpp b/Engine/Source/Editor/UnrealEd/Private/UnrealEdMisc.cpp
--- a/Engine/Source/Editor/UnrealEd/Private/UnrealEdMisc.cpp
+++ b/Engine/Source/Editor/UnrealEd/Private/UnrealEdMisc.cpp
@@ -1110,6 +1110,11 @@ void FUnrealEdMisc::OnExit()
 	const FString& PendingProjName = FUnrealEdMisc::Get().GetPendingProjectName();
 	if( PendingProjName.Len() > 0 )
 	{
+#if WITH_EDITOR
+			// Prevent the Zen subprocess data path from being inherited by a post-close process spawn
+			// such as when the editor is restarting itself.
+			FPlatformMisc::SetEnvironmentVar(TEXT("UE-ZenSubprocessDataPath"), nullptr);
+#endif // WITH_EDITOR
 		bool bSuccess = false;
 		if (FEditorDelegates::OnRestartRequested.IsBound())
 		{
