Change default zen data and install directories to be per-user, and clean up old default directories when migrating to a new one.

diff --git a/Engine/Config/BaseEngine.ini b/Engine/Config/BaseEngine.ini
--- a/Engine/Config/BaseEngine.ini
+++ b/Engine/Config/BaseEngine.ini
@@ -2438,7 +2438,7 @@ PathToRAC=RuntimeAssetCache
 AutoLaunch=true
 
 [Zen.AutoLaunch]
-DataPath=%APPSETTINGSDIR%Zen/Data
+DataPath=%ENGINEVERSIONAGNOSTICINSTALLEDUSERDIR%Zen/Data
 LimitProcessLifetime=true
 LocalDataCachePathEnvOverride=UE-LocalDataCachePath
 LocalDataCachePathCommandLineOverride=LocalDataCachePath
diff --git a/Engine/Source/Developer/Zen/Private/ZenServerInterface.cpp b/Engine/Source/Developer/Zen/Private/ZenServerInterface.cpp
--- a/Engine/Source/Developer/Zen/Private/ZenServerInterface.cpp
+++ b/Engine/Source/Developer/Zen/Private/ZenServerInterface.cpp
@@ -607,8 +607,14 @@ DetermineLocalDataCachePath(const TCHAR* ConfigSection, FString& DataPath)
 	}
 }
 
+static FString
+GetLocalZenRootPath()
+{
+	return FPaths::ConvertRelativePathToFull(FPaths::Combine(FPlatformProcess::UserSettingsDir(), *FApp::GetEpicProductIdentifier(), TEXT("Common")) + TEXT("/"));
+}
+
 static bool
-DetermineDataPath(const TCHAR* ConfigSection, FString& DataPath, bool& HasInvalidPathConfigurations)
+DetermineDataPath(const TCHAR* ConfigSection, FString& DataPath, bool& bHasInvalidPathConfigurations, bool& bIsDefaultDataPath)
 {
 	auto ValidateDataPath = [](const FString& InDataPath)
 	{
@@ -651,7 +657,7 @@ DetermineDataPath(const TCHAR* ConfigSection, FString& DataPath, bool& HasInvali
 			return true;
 		}
 		UE_LOG(LogZenServiceInstance, Warning, TEXT("Skipping command line override ZenDataPath=%s due to an invalid path"), *CommandLineOverrideValue);
-		HasInvalidPathConfigurations = true;
+		bHasInvalidPathConfigurations = true;
 	}
 
 	// Zen subprocess environment
@@ -664,7 +670,7 @@ DetermineDataPath(const TCHAR* ConfigSection, FString& DataPath, bool& HasInvali
 			return true;
 		}
 		UE_LOG(LogZenServiceInstance, Warning, TEXT("Skipping subprocess environment variable UE-ZenSubprocessDataPath=%s due to an invalid path"), *SubprocessDataPathEnvOverrideValue);
-		HasInvalidPathConfigurations = true;
+		bHasInvalidPathConfigurations = true;
 	}
 
 	// Zen registry/stored
@@ -678,7 +684,7 @@ DetermineDataPath(const TCHAR* ConfigSection, FString& DataPath, bool& HasInvali
 			return true;
 		}
 		UE_LOG(LogZenServiceInstance, Warning, TEXT("Skipping registry key Zen DataPath=%s due to an invalid path"), *DataPathEnvOverrideValue);
-		HasInvalidPathConfigurations = true;
+		bHasInvalidPathConfigurations = true;
 	}
 
 	// Zen environment
@@ -691,7 +697,7 @@ DetermineDataPath(const TCHAR* ConfigSection, FString& DataPath, bool& HasInvali
 			return true;
 		}
 		UE_LOG(LogZenServiceInstance, Warning, TEXT("Skipping environment variable UE-ZenDataPath=%s due to an invalid path"), *ZenDataPathEnvOverrideValue);
-		HasInvalidPathConfigurations = true;
+		bHasInvalidPathConfigurations = true;
 	}
 
 	// Follow local DDC (if outside workspace)
@@ -707,7 +713,7 @@ DetermineDataPath(const TCHAR* ConfigSection, FString& DataPath, bool& HasInvali
 			return true;
 		}
 		UE_LOG(LogZenServiceInstance, Warning, TEXT("Skipping local data cache path=%s due to an invalid path"), *LocalDataCachePath);
-		HasInvalidPathConfigurations = true;
+		bHasInvalidPathConfigurations = true;
 	}
 
 	// Zen config default
@@ -715,14 +721,16 @@ DetermineDataPath(const TCHAR* ConfigSection, FString& DataPath, bool& HasInvali
 	GConfig->GetString(ConfigSection, TEXT("DataPath"), ConfigDefaultPath, GEngineIni);
 	if (!ConfigDefaultPath.IsEmpty())
 	{
+		ConfigDefaultPath.ReplaceInline(TEXT("%ENGINEVERSIONAGNOSTICINSTALLEDUSERDIR%"), *GetLocalZenRootPath());
 		if (FString Path = ValidateDataPath(ConfigDefaultPath); !Path.IsEmpty())
 		{
 			DataPath = Path;
+			bIsDefaultDataPath = true;
 			UE_LOG(LogZenServiceInstance, Log, TEXT("Found Zen config default=%s"), *ConfigDefaultPath);
 			return true;
 		}
 		UE_LOG(LogZenServiceInstance, Warning, TEXT("Skipping Zen config default=%s due to an invalid path"), *ConfigDefaultPath);
-		HasInvalidPathConfigurations = true;
+		bHasInvalidPathConfigurations = true;
 	}
 	UE_LOG(LogZenServiceInstance, Warning, TEXT("Unable to determine a valid Zen data path"));
 	return false;
@@ -861,14 +869,14 @@ FServiceSettings::ReadFromConfig()
 			SettingsVariant.Emplace<FServiceAutoLaunchSettings>();
 			FServiceAutoLaunchSettings& AutoLaunchSettings = SettingsVariant.Get<FServiceAutoLaunchSettings>();
 
-			bool HasInvalidPathConfigurations = false;
-			if (!DetermineDataPath(AutoLaunchConfigSection, AutoLaunchSettings.DataPath, HasInvalidPathConfigurations))
+			bool bHasInvalidPathConfigurations = false;
+			if (!DetermineDataPath(AutoLaunchConfigSection, AutoLaunchSettings.DataPath, bHasInvalidPathConfigurations, AutoLaunchSettings.bIsDefaultDataPath))
 			{
 				PromptUserUnableToDetermineValidDataPath();
 				FPlatformMisc::RequestExit(true);
 				return;
 			}
-			else if (HasInvalidPathConfigurations)
+			else if (bHasInvalidPathConfigurations)
 			{
 				PromptUserAboutInvalidValidDataPathConfiguration(AutoLaunchSettings.DataPath);
 			}
@@ -941,6 +949,7 @@ FServiceSettings::ReadFromCompactBinary(FCbFieldView Field)
 				AutoLaunchSettings.ExtraArgs = FString(AutoLaunchSettingsObject["ExtraArgs"].AsString());
 				AutoLaunchSettings.DesiredPort = AutoLaunchSettingsObject["DesiredPort"].AsInt16();
 				AutoLaunchSettings.bShowConsole = AutoLaunchSettingsObject["ShowConsole"].AsBool();
+				AutoLaunchSettings.bIsDefaultDataPath = AutoLaunchSettingsObject["IsDefaultDataPath"].AsBool();
 				AutoLaunchSettings.bLimitProcessLifetime = AutoLaunchSettingsObject["LimitProcessLifetime"].AsBool();
 				ApplyProcessLifetimeOverride(AutoLaunchSettings.bLimitProcessLifetime);
 				AutoLaunchSettings.bSendUnattendedBugReports = AutoLaunchSettingsObject["SendUnattendedBugReports"].AsBool();
@@ -999,6 +1008,7 @@ FServiceSettings::WriteToCompactBinary(FCbWriter& Writer) const
 		Writer << "ExtraArgs" <<AutoLaunchSettings.ExtraArgs;
 		Writer << "DesiredPort" << AutoLaunchSettings.DesiredPort;
 		Writer << "ShowConsole" << AutoLaunchSettings.bShowConsole;
+		Writer << "IsDefaultDataPath" << AutoLaunchSettings.bIsDefaultDataPath;
 		Writer << "LimitProcessLifetime" << AutoLaunchSettings.bLimitProcessLifetime;
 		Writer << "SendUnattendedBugReports" << AutoLaunchSettings.bSendUnattendedBugReports;
 		Writer << "IsDefaultSharedRunContext" << AutoLaunchSettings.bIsDefaultSharedRunContext;
@@ -1179,7 +1189,7 @@ FString GetLocalServiceExecutableName()
 FString
 GetLocalServiceInstallPath()
 {
-	return FPaths::ConvertRelativePathToFull(FPaths::Combine(FPlatformProcess::ApplicationSettingsDir(), TEXT("Zen\\Install"), GetLocalServiceExecutableName()));
+	return FPaths::ConvertRelativePathToFull(FPaths::Combine(GetLocalZenRootPath(), TEXT("Zen\\Install"), GetLocalServiceExecutableName()));
 }
 
 static bool
@@ -1799,7 +1809,7 @@ GetLocalServiceInstallVersion(bool bDetailed)
 FString
 GetLocalInstallUtilityPath()
 {
-	return FPaths::ConvertRelativePathToFull(FPaths::Combine(FPlatformProcess::ApplicationSettingsDir(), TEXT("Zen\\Install"),
+	return FPaths::ConvertRelativePathToFull(FPaths::Combine(GetLocalZenRootPath(), TEXT("Zen\\Install"),
 #if PLATFORM_WINDOWS
 		TEXT("zen.exe")
 #else
@@ -2187,7 +2197,7 @@ FZenServiceInstance::ConditionalUpdateLocalInstall()
 #endif
 
 	FString InTreeCrashpadHandlerFilePath = FPaths::ConvertRelativePathToFull(FPlatformProcess::GenerateApplicationPath(TEXT("crashpad_handler"), EBuildConfiguration::Development));
-	FString InstallCrashpadHandlerFilePath = FPaths::ConvertRelativePathToFull(FPaths::Combine(FPlatformProcess::ApplicationSettingsDir(), TEXT("Zen\\Install"), FString(FPathViews::GetCleanFilename(InTreeCrashpadHandlerFilePath))));
+	FString InstallCrashpadHandlerFilePath = FPaths::ConvertRelativePathToFull(FPaths::Combine(GetLocalZenRootPath(), TEXT("Zen\\Install"), FString(FPathViews::GetCleanFilename(InTreeCrashpadHandlerFilePath))));
 
 	if (FileManager.FileExists(*InTreeCrashpadHandlerFilePath) && (bMainExecutablesUpdated || !FileManager.FileExists(*InstallCrashpadHandlerFilePath)))
 	{
@@ -2328,6 +2338,30 @@ FZenServiceInstance::AutoLaunch(const FServiceAutoLaunchSettings& InSettings, FS
 	// When not limiting process lifetime, only launch if the process is not already live.
 	if (bLaunchNewInstance)
 	{
+		if (InSettings.bIsDefaultDataPath && InSettings.bIsDefaultSharedRunContext)
+		{
+			// See if the default data path is migrating, and if so, clean up after the old one.
+			// Non-default data paths don't do the same thing because users are free to switch them back and forth
+			// and expext the contents to remain when they change.  Only the default one cleans up after itself
+			// to avoid a situation wherey the accumulate over time as the default location changes in config.
+			// This cleanup is best-effort and may fail if an instance is unexpectedly still using the previous path.
+			EnsureEditorSettingsConfigLoaded();
+			FString InUseDefaultDataPath;
+			if (!GConfig->GetString(TEXT("/Script/UnrealEd.ZenServerSettings"), TEXT("InUseDefaultDataPath"), InUseDefaultDataPath, GEditorSettingsIni))
+			{
+				InUseDefaultDataPath = FPaths::ConvertRelativePathToFull(FPaths::Combine(FPlatformProcess::ApplicationSettingsDir(), TEXT("Zen\\Data")));
+			}
+			if (!InUseDefaultDataPath.IsEmpty())
+			{
+				const FString InUseLockFilePath = FPaths::Combine(InUseDefaultDataPath, TEXT(".lock"));
+				if (!FPaths::IsSamePath(InUseDefaultDataPath, InSettings.DataPath) && !IsZenProcessUsingDataDir(*InUseLockFilePath, nullptr))
+				{
+					UE_LOG(LogZenServiceInstance, Display, TEXT("Migrating default data path from '%s' to '%s'.  Old location will be deleted."), *InUseDefaultDataPath, *InSettings.DataPath);
+					IFileManager::Get().DeleteDirectory(*InUseDefaultDataPath, false, true);
+				}
+			}
+		}
+
 		FString ParmsWithoutTransients = DetermineCmdLineWithoutTransientComponents(InSettings, InSettings.DesiredPort);
 		FString TransientParms;
 
@@ -2371,6 +2405,11 @@ FZenServiceInstance::AutoLaunch(const FServiceAutoLaunchSettings& InSettings, FS
 		}
 	}
 
+	if (InSettings.bIsDefaultDataPath && InSettings.bIsDefaultSharedRunContext)
+	{
+		GConfig->SetString(TEXT("/Script/UnrealEd.ZenServerSettings"), TEXT("InUseDefaultDataPath"), *InSettings.DataPath, GEditorSettingsIni);
+	}
+
 	OutHostName = TEXT("[::1]");
 	// Default to assuming that we get to run on the port we want
 	OutPort = InSettings.DesiredPort;
diff --git a/Engine/Source/Developer/Zen/Public/Experimental/ZenServerInterface.h b/Engine/Source/Developer/Zen/Public/Experimental/ZenServerInterface.h
--- a/Engine/Source/Developer/Zen/Public/Experimental/ZenServerInterface.h
+++ b/Engine/Source/Developer/Zen/Public/Experimental/ZenServerInterface.h
@@ -42,6 +42,7 @@ struct FServiceAutoLaunchSettings
 	FString ExtraArgs;
 	uint16 DesiredPort = 8558;
 	bool bShowConsole = false;
+	bool bIsDefaultDataPath = false;
 	bool bLimitProcessLifetime = false;
 	bool bSendUnattendedBugReports = false;
 	bool bIsDefaultSharedRunContext = true;
