Identify zenserver as Unreal Zen Storage Server in messagebox dialogs and accompanying logs for high priority messages.  Ensure google drive warning follows the same messagebox/log logic as other prompts.  Eliminate the use of RequestExit from within ZenServerInterface, failure should not force process shutdown.

diff --git a/Engine/Source/Developer/Zen/Private/ZenServerInterface.cpp b/Engine/Source/Developer/Zen/Private/ZenServerInterface.cpp
--- a/Engine/Source/Developer/Zen/Private/ZenServerInterface.cpp
+++ b/Engine/Source/Developer/Zen/Private/ZenServerInterface.cpp
@@ -457,14 +457,14 @@ PromptUserToSyncInTreeVersion(const FString& ServerFilePath)
 	if (!FApp::IsUnattended() && !IsRunningCommandlet() && !GIsRunningUnattendedScript)
 	{
 		FText ZenSyncSourcePromptTitle = NSLOCTEXT("Zen", "Zen_SyncSourcePromptTitle", "Failed to launch");
-		FText ZenSyncSourcePromptText = FText::Format(NSLOCTEXT("Zen", "Zen_SyncSourcePromptText", "ZenServer can not verify installation. Please make sure your source installation in properly synced at '{0}'"), FText::FromString(FPaths::GetPath(ServerFilePath)));
+		FText ZenSyncSourcePromptText = FText::Format(NSLOCTEXT("Zen", "Zen_SyncSourcePromptText", "Unreal Zen Storage Server can not verify installation. Please make sure your source installation in properly synced at '{0}'"), FText::FromString(FPaths::GetPath(ServerFilePath)));
 		FPlatformMisc::MessageBoxExt(EAppMsgType::Ok, *ZenSyncSourcePromptText.ToString(), *ZenSyncSourcePromptTitle.ToString());
 	}
 	else
 #endif
 	{
 		// Just log as there is no one to show a message
-		UE_LOG(LogZenServiceInstance, Display, TEXT("ZenServer can not verify installation. Please make sure your source installation in properly synced at '%s'"), *FPaths::GetPath(ServerFilePath));
+		UE_LOG(LogZenServiceInstance, Display, TEXT("Unreal Zen Storage Server can not verify installation. Please make sure your source installation in properly synced at '%s'"), *FPaths::GetPath(ServerFilePath));
 	}
 }
 
@@ -783,14 +783,14 @@ PromptUserUnableToDetermineValidDataPath()
 	{
 		FString LogDirPath = FPaths::ConvertRelativePathToFull(FPaths::ProjectLogDir());
 		FText ZenInvalidDataPathPromptTitle = NSLOCTEXT("Zen", "Zen_InvalidDataPathPromptTitle", "No Valid Data Path Configuration");
-		FText ZenInvalidDataPathPromptText = FText::Format(NSLOCTEXT("Zen", "Zen_InvalidDataPathPromptText", "ZenServer can not determine a valid data path.\nPlease check the log in '{0}' for details.\nUpdate your configuration and restart."), FText::FromString(LogDirPath));
+		FText ZenInvalidDataPathPromptText = FText::Format(NSLOCTEXT("Zen", "Zen_InvalidDataPathPromptText", "Unreal Zen Storage Server can not determine a valid data path.\nPlease check the log in '{0}' for details.\nUpdate your configuration and restart."), FText::FromString(LogDirPath));
 		FPlatformMisc::MessageBoxExt(EAppMsgType::Ok, *ZenInvalidDataPathPromptText.ToString(), *ZenInvalidDataPathPromptTitle.ToString());
 	}
 	else
 #endif
 	{
 		// Just log as there is no one to show a message
-		UE_LOG(LogZenServiceInstance, Warning, TEXT("ZenServer is unable to determine a valid data path"));
+		UE_LOG(LogZenServiceInstance, Warning, TEXT("Unreal Zen Storage Server is unable to determine a valid data path"));
 	}
 }
 
@@ -802,14 +802,14 @@ PromptUserAboutInvalidValidDataPathConfiguration(const FString& UsedDataPath)
 	{
 		FString LogDirPath = FPaths::ConvertRelativePathToFull(FPaths::ProjectLogDir());
 		FText ZenInvalidValidDataPathConfigurationPromptTitle = NSLOCTEXT("Zen", "Zen_InvalidValidDataPathConfigurationPromptTitle", "Invalid Data Paths");
-		FText ZenInvalidValidDataPathConfigurationPromptText = FText::Format(NSLOCTEXT("Zen", "Zen_InvalidValidDataPathConfigurationPromptText", "ZenServer has detected invalid data path configuration.\nPlease check the log in '{0}' for details.\n\nFalling back to using '{1}' as data path."), FText::FromString(LogDirPath), FText::FromString(UsedDataPath));
+		FText ZenInvalidValidDataPathConfigurationPromptText = FText::Format(NSLOCTEXT("Zen", "Zen_InvalidValidDataPathConfigurationPromptText", "Unreal Zen Storage Server has detected invalid data path configuration.\nPlease check the log in '{0}' for details.\n\nFalling back to using '{1}' as data path."), FText::FromString(LogDirPath), FText::FromString(UsedDataPath));
 		FPlatformMisc::MessageBoxExt(EAppMsgType::Ok, *ZenInvalidValidDataPathConfigurationPromptText.ToString(), *ZenInvalidValidDataPathConfigurationPromptTitle.ToString());
 	}
 	else
 #endif
 	{
 		// Just log as there is no one to show a message
-		UE_LOG(LogZenServiceInstance, Warning, TEXT("ZenServer has detected invalid data path configuration. Falling back to '%s'"), *UsedDataPath);
+		UE_LOG(LogZenServiceInstance, Warning, TEXT("Unreal Zen Storage Server has detected invalid data path configuration. Falling back to '%s'"), *UsedDataPath);
 	}
 }
 
@@ -817,16 +817,20 @@ PromptUserAboutInvalidValidDataPathConfiguration(const FString& UsedDataPath)
 static void
 PromptUserIsUsingGoogleDriveAsDataPath()
 {
-	if (FApp::IsUnattended())
+#if !IS_PROGRAM
+	if (!FApp::IsUnattended() && !IsRunningCommandlet() && !GIsRunningUnattendedScript)
+	{
+		FString LogDirPath = FPaths::ConvertRelativePathToFull(FPaths::ProjectLogDir());
+		FText ZenInvalidDataPathPromptTitle = NSLOCTEXT("Zen", "Zen_GoogleDriveDataPathPromptTitle", "Using Google Drive as a data path");
+		FText ZenInvalidDataPathPromptText = FText::Format(NSLOCTEXT("Zen", "Zen_GoogleDriveDataPathPromptText", "Unreal Zen Storage Server is configured to use Google Drive as a data path, this is highly inadvisable.\nPlease use a data path on a local physical drive.\nCheck the log in '{0}' for details.\nUpdate your configuration and restart."), FText::FromString(LogDirPath));
+		FPlatformMisc::MessageBoxExt(EAppMsgType::Ok, *ZenInvalidDataPathPromptText.ToString(), *ZenInvalidDataPathPromptTitle.ToString());
+	}
+	else
+#endif
 	{
 		// Just log as there is no one to show a message
-		UE_LOG(LogZenServiceInstance, Warning, TEXT("ZenServer is configured to use Google Drive as a data path, this is highly inadvisable. Please use a path on a local physical drive."));
-		return;
+		UE_LOG(LogZenServiceInstance, Warning, TEXT("Unreal Zen Storage Server is configured to use Google Drive as a data path, this is highly inadvisable. Please use a path on a local physical drive."));
 	}
-	FString LogDirPath = FPaths::ConvertRelativePathToFull(FPaths::ProjectLogDir());
-	FText ZenInvalidDataPathPromptTitle = NSLOCTEXT("Zen", "Zen_GoogleDriveDataPathPromptTitle", "Using Google Drive as a data path");
-	FText ZenInvalidDataPathPromptText = FText::Format(NSLOCTEXT("Zen", "Zen_GoogleDriveDataPathPromptText", "ZenServer is configured to use Google Drive as a data path, this is highly inadvisable.\nPlease use a data path on a local physical drive.\nCheck the log in '{0}' for details.\nUpdate your configuration and restart."), FText::FromString(LogDirPath));
-	FPlatformMisc::MessageBoxExt(EAppMsgType::Ok, *ZenInvalidDataPathPromptText.ToString(), *ZenInvalidDataPathPromptTitle.ToString());
 }
 #endif // PLATFORM_WINDOWS
 
@@ -852,7 +856,7 @@ static void ReadCbField(FCbFieldView Field, UE::Zen::FZenCIDStats& OutValue)
 	ReadCbField(ObjectView["size"], OutValue.Size);
 }
 
-void
+bool
 FServiceSettings::ReadFromConfig()
 {
 	check(GConfig && GConfig->IsReadyForUse());
@@ -873,8 +877,7 @@ FServiceSettings::ReadFromConfig()
 			if (!DetermineDataPath(AutoLaunchConfigSection, AutoLaunchSettings.DataPath, bHasInvalidPathConfigurations, AutoLaunchSettings.bIsDefaultDataPath))
 			{
 				PromptUserUnableToDetermineValidDataPath();
-				FPlatformMisc::RequestExit(true);
-				return;
+				return false;
 			}
 			else if (bHasInvalidPathConfigurations)
 			{
@@ -931,9 +934,10 @@ FServiceSettings::ReadFromConfig()
 		GConfig->GetString(ConnectExistingConfigSection, TEXT("HostName"), ConnectExistingSettings.HostName, GEngineIni);
 		ReadUInt16FromConfig(ConnectExistingConfigSection, TEXT("Port"), ConnectExistingSettings.Port, GEngineIni);
 	}
+	return true;
 }
 
-void
+bool
 FServiceSettings::ReadFromCompactBinary(FCbFieldView Field)
 {
 	if (bool bAutoLaunchValue = Field["bAutoLaunch"].AsBool())
@@ -968,9 +972,10 @@ FServiceSettings::ReadFromCompactBinary(FCbFieldView Field)
 			ConnectExistingSettings.Port = ConnectExistingSettingsObject["Port"].AsInt16();
 		}
 	}
+	return true;
 }
 
-void
+bool
 FServiceSettings::ReadFromURL(FStringView InstanceURL)
 {
 	SettingsVariant.Emplace<FServiceConnectSettings>();
@@ -993,6 +998,7 @@ FServiceSettings::ReadFromURL(FStringView InstanceURL)
 		ConnectExistingSettings.HostName = InstanceURL;
 		ConnectExistingSettings.Port = 8558;
 	}
+	return true;
 }
 
 void
@@ -1920,18 +1926,18 @@ FZenServiceInstance::IsServiceReady()
 		
 		if (Result == Zen::FZenHttpRequest::Result::Success && Zen::IsSuccessCode(Request.GetResponseCode()))
 		{
-			UE_LOG(LogZenServiceInstance, Display, TEXT("ZenServer HTTP service at %s status: %s."), ZenDomain.ToString(), *Request.GetResponseAsString());
+			UE_LOG(LogZenServiceInstance, Display, TEXT("Unreal Zen Storage Server HTTP service at %s status: %s."), ZenDomain.ToString(), *Request.GetResponseAsString());
 			return true;
 		}
 		else
 		{
 			if (IsServiceRunningLocally())
 			{
-				UE_LOG(LogZenServiceInstance, Warning, TEXT("Unable to reach ZenServer HTTP service at %s. Status: %d . Response: %s"), ZenDomain.ToString(), Request.GetResponseCode(), *Request.GetResponseAsString());
+				UE_LOG(LogZenServiceInstance, Warning, TEXT("Unable to reach Unreal Zen Storage Server HTTP service at %s. Status: %d . Response: %s"), ZenDomain.ToString(), Request.GetResponseCode(), *Request.GetResponseAsString());
 			}
 			else
 			{
-				UE_LOG(LogZenServiceInstance, Display, TEXT("Unable to reach ZenServer HTTP service at %s. Status: %d . Response: %s"), ZenDomain.ToString(), Request.GetResponseCode(), *Request.GetResponseAsString());
+				UE_LOG(LogZenServiceInstance, Display, TEXT("Unable to reach Unreal Zen Storage Server HTTP service at %s. Status: %d . Response: %s"), ZenDomain.ToString(), Request.GetResponseCode(), *Request.GetResponseAsString());
 			}
 		}
 	}
@@ -2098,14 +2104,14 @@ PromptUserToStopRunningServerInstanceForUpdate(const FString& ServerFilePath)
 	if (!FApp::IsUnattended() && !IsRunningCommandlet() && !GIsRunningUnattendedScript)
 	{
 		FText ZenUpdatePromptTitle = NSLOCTEXT("Zen", "Zen_UpdatePromptTitle", "Update required");
-		FText ZenUpdatePromptText = FText::Format(NSLOCTEXT("Zen", "Zen_UpdatePromptText", "ZenServer needs to be updated to a new version. Please shut down Unreal Editor and any tools that are using the ZenServer at '{0}'"), FText::FromString(ServerFilePath));
+		FText ZenUpdatePromptText = FText::Format(NSLOCTEXT("Zen", "Zen_UpdatePromptText", "Unreal Zen Storage Server needs to be updated to a new version. Please shut down Unreal Editor and any tools that are using the ZenServer at '{0}'"), FText::FromString(ServerFilePath));
 		FPlatformMisc::MessageBoxExt(EAppMsgType::Ok, *ZenUpdatePromptText.ToString(), *ZenUpdatePromptTitle.ToString());
 	}
 	else
 #endif
 	{
 		// Just log as there is no one to show a message
-		UE_LOG(LogZenServiceInstance, Display, TEXT("ZenServer needs to be updated to a new version. Please shut down any tools that are using the ZenServer at '%s'"), *ServerFilePath);
+		UE_LOG(LogZenServiceInstance, Display, TEXT("Unreal Zen Storage Server needs to be updated to a new version. Please shut down any tools that are using the ZenServer at '%s'"), *ServerFilePath);
 	}
 }
 
@@ -2116,14 +2122,14 @@ PromptUserOfLockedDataFolder(const FString& DataPath)
 	if (!FApp::IsUnattended() && !IsRunningCommandlet() && !GIsRunningUnattendedScript)
 	{
 		FText ZenLaunchFailurePromptTitle = NSLOCTEXT("Zen", "Zen_NonLocalProcessUsesDataDirPromptTitle", "Failed to launch");
-		FText ZenLaunchFailurePromptText = FText::Format(NSLOCTEXT("Zen", "Zen_NonLocalProcessUsesDataDirPromptText", "ZenServer Failed to auto launch, an unknown process is locking the data folder '{0}'"), FText::FromString(DataPath));
+		FText ZenLaunchFailurePromptText = FText::Format(NSLOCTEXT("Zen", "Zen_NonLocalProcessUsesDataDirPromptText", "Unreal Zen Storage Server Failed to auto launch, an unknown process is locking the data folder '{0}'"), FText::FromString(DataPath));
 		FPlatformMisc::MessageBoxExt(EAppMsgType::Ok, *ZenLaunchFailurePromptText.ToString(), *ZenLaunchFailurePromptTitle.ToString());
 	}
 	else
 #endif
 	{
 		// Just log as there is no one to show a message
-		UE_LOG(LogZenServiceInstance, Warning, TEXT("ZenServer Failed to auto launch, an unknown process is locking the data folder '%s'"), *DataPath);
+		UE_LOG(LogZenServiceInstance, Warning, TEXT("Unreal Zen Storage Server Failed to auto launch, an unknown process is locking the data folder '%s'"), *DataPath);
 	}
 }
 
@@ -2134,15 +2140,14 @@ PromptUserOfFailedShutDownOfExistingProcess(uint16 Port)
 	if (!FApp::IsUnattended() && !IsRunningCommandlet() && !GIsRunningUnattendedScript)
 	{
 		FText ZenLaunchFailurePromptTitle = NSLOCTEXT("Zen", "Zen_ShutdownFailurePromptTitle", "Failed to launch");
-		FText ZenLaunchFailurePromptText = FText::Format(NSLOCTEXT("Zen", "Zen_ShutdownFailurePromptText", "ZenServer Failed to auto launch, failed to shut down currently running service using port '{0}'"), FText::AsNumber(Port, &FNumberFormattingOptions::DefaultNoGrouping()));
+		FText ZenLaunchFailurePromptText = FText::Format(NSLOCTEXT("Zen", "Zen_ShutdownFailurePromptText", "Unreal Zen Storage Server Failed to auto launch, failed to shut down currently running service using port '{0}'"), FText::AsNumber(Port, &FNumberFormattingOptions::DefaultNoGrouping()));
 		FPlatformMisc::MessageBoxExt(EAppMsgType::Ok, *ZenLaunchFailurePromptText.ToString(), *ZenLaunchFailurePromptTitle.ToString());
 	}
 	else
 #endif
 	{
 		// Just log as there is no one to show a message
-		UE_LOG(LogZenServiceInstance, Warning, TEXT("ZenServer Failed to auto launch, failed to shut down currently running service using port %u"), Port);
-		FPlatformMisc::RequestExit(true);
+		UE_LOG(LogZenServiceInstance, Warning, TEXT("Unreal Zen Storage Server Failed to auto launch, failed to shut down currently running service using port %u"), Port);
 	}
 }
 
@@ -2243,7 +2248,6 @@ FZenServiceInstance::AutoLaunch(const FServiceAutoLaunchSettings& InSettings, FS
 			{
 				UE_LOG(LogZenServiceInstance, Warning, TEXT("Found locked valid lock file '%s' but no matching process (Pid: %d), exiting"), *LockFilePath, LockFileState.ProcessId);
 				PromptUserOfLockedDataFolder(*InSettings.DataPath);
-				FPlatformMisc::RequestExit(true);
 				return false;
 			}
 		}
@@ -2319,7 +2323,6 @@ FZenServiceInstance::AutoLaunch(const FServiceAutoLaunchSettings& InSettings, FS
 		if (!ShutDownZenServerProcessLockingDataDir(InSettings.DataPath))
 		{
 			PromptUserOfFailedShutDownOfExistingProcess(InSettings.DesiredPort);
-			FPlatformMisc::RequestExit(true);
 			return false;
 		}
 	}
@@ -2329,7 +2332,6 @@ FZenServiceInstance::AutoLaunch(const FServiceAutoLaunchSettings& InSettings, FS
 		if (!ShutdownZenServerProcess(ShutdownExistingInstanceForPid))
 		{
 			PromptUserOfFailedShutDownOfExistingProcess(InSettings.DesiredPort);
-			FPlatformMisc::RequestExit(true);
 			return false;
 		}
 	}
@@ -2457,16 +2459,15 @@ FZenServiceInstance::AutoLaunch(const FServiceAutoLaunchSettings& InSettings, FS
 						FString LogFilePath = FPaths::Combine(InSettings.DataPath, TEXT("logs"), TEXT("zenserver.log"));
 						FPaths::MakePlatformFilename(LogFilePath);
 						FormatArguments.Add(TEXT("LogFilePath"), FText::FromString(LogFilePath));
-						FText ZenLaunchFailurePromptText = FText::Format(NSLOCTEXT("Zen", "Zen_LaunchFailurePromptText", "ZenServer failed to launch. This process will now exit. Please check the ZenServer log file for details:\n{LogFilePath}"), FormatArguments);
+						FText ZenLaunchFailurePromptText = FText::Format(NSLOCTEXT("Zen", "Zen_LaunchFailurePromptText", "Unreal Zen Storage Server failed to launch. Please check the ZenServer log file for details:\n{LogFilePath}"), FormatArguments);
 						FPlatformMisc::MessageBoxExt(EAppMsgType::Ok, *ZenLaunchFailurePromptText.ToString(), *ZenLaunchFailurePromptTitle.ToString());
-						FPlatformMisc::RequestExit(true);
 						return false;
 					}
 					else
 #endif
 					{
 						// Just log as there is no one to show a message
-						UE_LOG(LogZenServiceInstance, Warning, TEXT("ZenServer did not launch in the expected duration"));
+						UE_LOG(LogZenServiceInstance, Warning, TEXT("Unreal Zen Storage Server did not launch in the expected duration"));
 						return false;
 					}
 				}
@@ -2480,7 +2481,7 @@ FZenServiceInstance::AutoLaunch(const FServiceAutoLaunchSettings& InSettings, FS
 			else if (!(FApp::IsUnattended() || IsRunningCommandlet() || GIsRunningUnattendedScript) && ZenWaitDuration > 20.0 && (DurationPhase == EWaitDurationPhase::Medium))
 			{
 				FText ZenLongWaitPromptTitle = NSLOCTEXT("Zen", "Zen_LongWaitPromptTitle", "Wait for ZenServer?");
-				FText ZenLongWaitPromptText = NSLOCTEXT("Zen", "Zen_LongWaitPromptText", "ZenServer is taking a long time to launch. It may be performing maintenance. Keep waiting?");
+				FText ZenLongWaitPromptText = NSLOCTEXT("Zen", "Zen_LongWaitPromptText", "Unreal Zen Storage Server is taking a long time to launch. It may be performing maintenance. Keep waiting?");
 				if (FPlatformMisc::MessageBoxExt(EAppMsgType::YesNo, *ZenLongWaitPromptText.ToString(), *ZenLongWaitPromptTitle.ToString()) == EAppReturnType::No)
 				{
 					return false;
diff --git a/Engine/Source/Developer/Zen/Public/Experimental/ZenServerInterface.h b/Engine/Source/Developer/Zen/Public/Experimental/ZenServerInterface.h
--- a/Engine/Source/Developer/Zen/Public/Experimental/ZenServerInterface.h
+++ b/Engine/Source/Developer/Zen/Public/Experimental/ZenServerInterface.h
@@ -55,9 +55,9 @@ struct FServiceSettings
 	inline bool IsAutoLaunch() const { return SettingsVariant.IsType<FServiceAutoLaunchSettings>(); }
 	inline bool IsConnectExisting() const { return SettingsVariant.IsType<FServiceConnectSettings>(); }
 
-	UE_API void ReadFromConfig();
-	UE_API void ReadFromCompactBinary(FCbFieldView Field);
-	UE_API void ReadFromURL(FStringView InstanceURL);
+	UE_API bool ReadFromConfig();
+	UE_API bool ReadFromCompactBinary(FCbFieldView Field);
+	UE_API bool ReadFromURL(FStringView InstanceURL);
 
 	UE_API void WriteToCompactBinary(FCbWriter& Writer) const;
 
