Attempt shut down of running zenserver process if we find a locked lock file but the pid registered in the lock file can't be found.
If a pid can not be found when finding a matching desired/effective port, keep looking if any other state has a matching port.
Hardened and added logging to NativeIsProcessRunning.

Run of Low Level Tests in Main  here: https://horde.devtools.epicgames.com/job/65e81c5d156d6df40d304906

diff --git a/Engine/Source/Developer/Zen/Private/ZenServerInterface.cpp b/Engine/Source/Developer/Zen/Private/ZenServerInterface.cpp
--- a/Engine/Source/Developer/Zen/Private/ZenServerInterface.cpp
+++ b/Engine/Source/Developer/Zen/Private/ZenServerInterface.cpp
@@ -75,20 +75,48 @@ static bool NativeIsProcessRunning(uint32 Pid)
 		{
 			return false;
 		}
+		else if (Error == ERROR_ACCESS_DENIED)
+		{
+			UE_LOG(LogZenServiceInstance, Warning, TEXT("No access to open running process %d: %d, assuming it is running"), Pid, Error);
+			return true;
+		}
+		UE_LOG(LogZenServiceInstance, Warning, TEXT("Failed to open running process %d: %d, assuming it is not running"), Pid, Error);
 		return false;
 	}
 	ON_SCOPE_EXIT{ CloseHandle(Handle); };
 
-	bool  bStillActive = true;
 	DWORD ExitCode = 0;
-	if (0 != GetExitCodeProcess(Handle, &ExitCode))
+	if (GetExitCodeProcess(Handle, &ExitCode) == 0)
 	{
-		bStillActive = ExitCode == STILL_ACTIVE;
+		DWORD Error = GetLastError();
+		UE_LOG(LogZenServiceInstance, Warning, TEXT("Failed to get running process exit code %d: %d, assuming it is still running"), Pid, Error);
+		return true;
+	}
+	else if (ExitCode == STILL_ACTIVE)
+	{
+		return true;
 	}
 
-	return bStillActive;
+	return false;
+
 #elif PLATFORM_UNIX || PLATFORM_MAC
-	return (kill(pid_t(Pid), 0) == 0);
+	int Res = kill(pid_t(Pid), 0);
+	if (Res == 0)
+	{
+		return true;
+	}
+	int Error = errno;
+	if (Error == EPERM)
+	{
+		UE_LOG(LogZenServiceInstance, Warning, TEXT("No permission to signal running process %d: %d, assuming it is running"), Pid, Error);
+		return true;
+	}
+	else if (Error == ESRCH)
+	{
+		return false;
+	}
+	UE_LOG(LogZenServiceInstance, Warning, TEXT("Failed to signal running process %d: %d, assuming it is running"), Pid, Error);
+	return true;
 #endif
 }
 
@@ -282,7 +310,6 @@ const ZenServerState::ZenServerEntry* ZenServerState::LookupByDesiredListenPortI
 			{
 				return Entry;
 			}
-			return nullptr;
 		}
 	}
 
@@ -316,7 +343,6 @@ const ZenServerState::ZenServerEntry* ZenServerState::LookupByEffectiveListenPor
 			{
 				return Entry;
 			}
-			return nullptr;
 		}
 	}
 
@@ -2242,60 +2268,58 @@ FZenServiceInstance::AutoLaunch(const FServiceAutoLaunchSettings& InSettings, FS
 	if (LockFileState.IsReady)
 	{
 		const ZenServerState State(/*ReadOnly*/true);
-		if (State.LookupByPid(LockFileState.ProcessId) == nullptr)
+		if (State.LookupByPid(LockFileState.ProcessId) == nullptr && IsZenProcessUsingDataDir(*LockFilePath, nullptr))
 		{
-			if (IsZenProcessUsingDataDir(*LockFilePath, nullptr))
-			{
-				UE_LOG(LogZenServiceInstance, Warning, TEXT("Found locked valid lock file '%s' but no matching process (Pid: %d), exiting"), *LockFilePath, LockFileState.ProcessId);
-				PromptUserOfLockedDataFolder(*InSettings.DataPath);
-				return false;
-			}
+			UE_LOG(LogZenServiceInstance, Warning, TEXT("Found locked valid lock file '%s' but can't find registered process (Pid: %d), will attempt shut down"), *LockFilePath, LockFileState.ProcessId);
+			bShutDownExistingInstanceForDataPath = true;
 		}
-
-		if (InSettings.bIsDefaultSharedRunContext)
+		else
 		{
-			FZenLocalServiceRunContext DesiredRunContext;
-			DesiredRunContext.Executable = ExecutablePath;
-			DesiredRunContext.CommandlineArguments = DetermineCmdLineWithoutTransientComponents(InSettings, InSettings.DesiredPort);
-			DesiredRunContext.WorkingDirectory = WorkingDirectory;
-			DesiredRunContext.DataPath = InSettings.DataPath;
-			DesiredRunContext.bShowConsole = InSettings.bShowConsole;
+			if (InSettings.bIsDefaultSharedRunContext)
+			{
+				FZenLocalServiceRunContext DesiredRunContext;
+				DesiredRunContext.Executable = ExecutablePath;
+				DesiredRunContext.CommandlineArguments = DetermineCmdLineWithoutTransientComponents(InSettings, InSettings.DesiredPort);
+				DesiredRunContext.WorkingDirectory = WorkingDirectory;
+				DesiredRunContext.DataPath = InSettings.DataPath;
+				DesiredRunContext.bShowConsole = InSettings.bShowConsole;
 
-			FZenLocalServiceRunContext CurrentRunContext;
+				FZenLocalServiceRunContext CurrentRunContext;
 
-			bool ReadCurrentContextOK = CurrentRunContext.ReadFromJsonFile(*ExecutionContextFilePath);
-			if (ReadCurrentContextOK && (DesiredRunContext == CurrentRunContext))
-			{
-				UE_LOG(LogZenServiceInstance, Log, TEXT("Found existing instance running on port %u matching our settings, no actions needed"), InSettings.DesiredPort);
-				bLaunchNewInstance = false;
-				bShutDownExistingInstanceForDataPath = false;
+				bool ReadCurrentContextOK = CurrentRunContext.ReadFromJsonFile(*ExecutionContextFilePath);
+				if (ReadCurrentContextOK && (DesiredRunContext == CurrentRunContext))
+				{
+					UE_LOG(LogZenServiceInstance, Log, TEXT("Found existing instance running on port %u matching our settings, no actions needed"), InSettings.DesiredPort);
+					bLaunchNewInstance = false;
+					bShutDownExistingInstanceForDataPath = false;
+				}
+				else
+				{
+					FString JsonTcharText;
+					{
+						TSharedRef<TJsonWriter<TCHAR, TPrettyJsonPrintPolicy<TCHAR>>> Writer = TJsonWriterFactory<TCHAR, TPrettyJsonPrintPolicy<TCHAR>>::Create(&JsonTcharText);
+						Writer->WriteObjectStart();
+						Writer->WriteObjectStart("Current");
+						CurrentRunContext.WriteToJson(*Writer);
+						Writer->WriteObjectEnd();
+						Writer->WriteObjectStart("Desired");
+						DesiredRunContext.WriteToJson(*Writer);
+						Writer->WriteObjectEnd();
+						Writer->WriteObjectEnd();
+						Writer->Close();
+					}
+					UE_LOG(LogZenServiceInstance, Log, TEXT("Found existing instance running on port %u with different run context, will attempt shut down\n{%s}"), InSettings.DesiredPort, *JsonTcharText);
+					bShutDownExistingInstanceForDataPath = true;
+					bLaunchNewInstance = true;
+				}
 			}
 			else
 			{
-				FString JsonTcharText;
-				{
-					TSharedRef<TJsonWriter<TCHAR, TPrettyJsonPrintPolicy<TCHAR>>> Writer = TJsonWriterFactory<TCHAR, TPrettyJsonPrintPolicy<TCHAR>>::Create(&JsonTcharText);
-					Writer->WriteObjectStart();
-					Writer->WriteObjectStart("Current");
-					CurrentRunContext.WriteToJson(*Writer);
-					Writer->WriteObjectEnd();
-					Writer->WriteObjectStart("Desired");
-					DesiredRunContext.WriteToJson(*Writer);
-					Writer->WriteObjectEnd();
-					Writer->WriteObjectEnd();
-					Writer->Close();
-				}
-				UE_LOG(LogZenServiceInstance, Log, TEXT("Found existing instance running on port %u with different run context, will attempt shut down\n{%s}"), InSettings.DesiredPort, *JsonTcharText);
-				bShutDownExistingInstanceForDataPath = true;
-				bLaunchNewInstance = true;
+				UE_LOG(LogZenServiceInstance, Log, TEXT("Found existing instance running on port %u when not using shared context, will use it"), InSettings.DesiredPort);
+				bShutDownExistingInstanceForDataPath = false;
+				bLaunchNewInstance = false;
 			}
 		}
-		else
-		{
-			UE_LOG(LogZenServiceInstance, Log, TEXT("Found existing instance running on port %u when not using shared context, will use it"), InSettings.DesiredPort);
-			bShutDownExistingInstanceForDataPath = false;
-			bLaunchNewInstance = false;
-		}
 	}
 	else
 	{
