Address zenserver autolaunch/lifecycle issues identified on Windows, Mac, and Linux using lifecycle stress test.

diff --git a/Engine/Source/Developer/Zen/Private/ZenServerInterface.cpp b/Engine/Source/Developer/Zen/Private/ZenServerInterface.cpp
--- a/Engine/Source/Developer/Zen/Private/ZenServerInterface.cpp
+++ b/Engine/Source/Developer/Zen/Private/ZenServerInterface.cpp
@@ -1492,6 +1492,11 @@ static bool ShutDownZenServerProcessLockingDataDir(const FString& DataPath, doub
 			uint32_t Pid = 0;
 			if (!FindZenProcessId(&Pid))
 			{
+				if (!IsLockFileLocked(*LockFilePath, true))
+				{
+					UE_LOG(LogZenServiceInstance, Display, TEXT("Lock file '%s' is no longer active, nothing to do"), *LockFilePath);
+					return true;
+				}
 				UE_LOG(LogZenServiceInstance, Warning, TEXT("Failed to find zenserver process locking file '%s'"), *LockFilePath);
 				return false;
 			}
@@ -1550,6 +1555,11 @@ static bool ShutDownZenServerProcessLockingDataDir(const FString& DataPath, doub
 		uint32_t Pid = 0;
 		if (!FindZenProcessId(&Pid))
 		{
+			if (!IsLockFileLocked(*LockFilePath, true))
+			{
+				UE_LOG(LogZenServiceInstance, Display, TEXT("Lock file '%s' is no longer active, nothing to do"), *LockFilePath);
+				return true;
+			}
 			UE_LOG(LogZenServiceInstance, Warning, TEXT("Failed to find zenserver process locking file '%s'"), *LockFilePath);
 			return false;
 		}
@@ -1943,7 +1953,8 @@ FZenServiceInstance::IsServiceRunning()
 bool 
 FZenServiceInstance::IsServiceReady()
 {
-	if (IsServiceRunning())
+	uint32 Attempt = 0;
+	while (IsServiceRunning())
 	{
 		TStringBuilder<128> ZenDomain;
 		ZenDomain << HostName << TEXT(":") << Port;
@@ -1955,17 +1966,21 @@ FZenServiceInstance::IsServiceReady()
 			UE_LOG(LogZenServiceInstance, Display, TEXT("Unreal Zen Storage Server HTTP service at %s status: %s."), ZenDomain.ToString(), *Request.GetResponseAsString());
 			return true;
 		}
-		else
+
+		if (IsServiceRunningLocally())
 		{
-			if (IsServiceRunningLocally())
+			if (Attempt > 4)
 			{
 				UE_LOG(LogZenServiceInstance, Warning, TEXT("Unable to reach Unreal Zen Storage Server HTTP service at %s. Status: %d . Response: %s"), ZenDomain.ToString(), Request.GetResponseCode(), *Request.GetResponseAsString());
-			}
-			else
-			{
-				UE_LOG(LogZenServiceInstance, Display, TEXT("Unable to reach Unreal Zen Storage Server HTTP service at %s. Status: %d . Response: %s"), ZenDomain.ToString(), Request.GetResponseCode(), *Request.GetResponseAsString());
+				break;
 			}
 		}
+		else
+		{
+			UE_LOG(LogZenServiceInstance, Display, TEXT("Unable to reach Unreal Zen Storage Server HTTP service at %s. Status: %d . Response: %s"), ZenDomain.ToString(), Request.GetResponseCode(), *Request.GetResponseAsString());
+			break;
+		}
+		Attempt++;
 	}
 	return false;
 }
@@ -2110,6 +2125,26 @@ FZenServiceInstance::Initialize()
 				}
 				AutoLaunchedPort = Port;
 				bIsRunningLocally = true;
+
+				const FTimespan MaximumWaitForHealth = FTimespan::FromSeconds(20);
+				FDateTime StartedWaitingForHealth = FDateTime::UtcNow();
+				bool bLastReadyResult = IsServiceReady();
+				while (!bLastReadyResult)
+				{
+					FTimespan WaitForHealth = FDateTime::UtcNow() - StartedWaitingForHealth;
+					if (WaitForHealth > MaximumWaitForHealth)
+					{
+						UE_LOG(LogZenServiceInstance, Warning, TEXT("Local ZenServer AutoLaunch initialization timed out waiting for service to become healthy"));
+						break;
+					}
+
+					FPlatformProcess::Sleep(0.5f);
+					if (!IsZenProcessUsingEffectivePort(Port))
+					{
+						AutoLaunch(Settings.SettingsVariant.Get<FServiceAutoLaunchSettings>(), *GetLocalServiceInstallPath(), HostName, Port);
+					}
+					bLastReadyResult = IsServiceReady();
+				}
 			}
 		}
 	}
@@ -2362,7 +2397,7 @@ FZenServiceInstance::AutoLaunch(const FServiceAutoLaunchSettings& InSettings, FS
 
 	// When limiting process lifetime, always re-launch to add sponsor process IDs.
 	// When not limiting process lifetime, only launch if the process is not already live.
-	if (bLaunchNewInstance)
+	if (bLaunchNewInstance || InSettings.bLimitProcessLifetime)
 	{
 		if (InSettings.bIsDefaultDataPath && InSettings.bIsDefaultSharedRunContext)
 		{
