Remove one line edit added in my previous changelist because it isn't necessary.

diff --git a/Engine/Source/Developer/Zen/Private/ZenServerInterface.cpp b/Engine/Source/Developer/Zen/Private/ZenServerInterface.cpp
--- a/Engine/Source/Developer/Zen/Private/ZenServerInterface.cpp
+++ b/Engine/Source/Developer/Zen/Private/ZenServerInterface.cpp
@@ -2397,7 +2397,7 @@ FZenServiceInstance::AutoLaunch(const FServiceAutoLaunchSettings& InSettings, FS
 
 	// When limiting process lifetime, always re-launch to add sponsor process IDs.
 	// When not limiting process lifetime, only launch if the process is not already live.
-	if (bLaunchNewInstance || InSettings.bLimitProcessLifetime)
+	if (bLaunchNewInstance)
 	{
 		if (InSettings.bIsDefaultDataPath && InSettings.bIsDefaultSharedRunContext)
 		{
