Address missing file descriptor close on Mac/Linux that is related to the lock file handling when launching zenserver.  Also fix an incorrect return value from the Mac/Linux implementation of IsLockFileLocked in the case of an orphaned lock file.

diff --git a/Engine/Source/Developer/Zen/Private/ZenServerInterface.cpp b/Engine/Source/Developer/Zen/Private/ZenServerInterface.cpp
--- a/Engine/Source/Developer/Zen/Private/ZenServerInterface.cpp
+++ b/Engine/Source/Developer/Zen/Private/ZenServerInterface.cpp
@@ -1189,6 +1189,7 @@ ReadCbLockFile(FStringView FileName)
 
 	if (errno != EWOULDBLOCK && errno != EAGAIN)
 	{
+		close(Fd);
 		return {};
 	}
 
@@ -1199,6 +1200,7 @@ ReadCbLockFile(FStringView FileName)
 	FUniqueBuffer FileBytes = FUniqueBuffer::Alloc(FileSize);
 	if (read(Fd, FileBytes.GetData(), FileSize) == FileSize)
 	{
+		close(Fd);
 		return ReadLockData(std::move(FileBytes));
 	}
 
@@ -1256,7 +1258,7 @@ IsLockFileLocked(const TCHAR* FileName, bool bAttemptCleanUp = false)
 	flock(Fd, LOCK_UN);
 	close(Fd);
 
-	return true;
+	return false;
 #endif
 }
 
