Further hardening of zen launch sequence.

diff --git a/Engine/Source/Developer/Zen/Private/ZenServerInterface.cpp b/Engine/Source/Developer/Zen/Private/ZenServerInterface.cpp
--- a/Engine/Source/Developer/Zen/Private/ZenServerInterface.cpp
+++ b/Engine/Source/Developer/Zen/Private/ZenServerInterface.cpp
@@ -2116,37 +2116,37 @@ FZenServiceInstance::Initialize()
 		FString ExecutableInstallPath = ConditionalUpdateLocalInstall();
 		if (!ExecutableInstallPath.IsEmpty())
 		{
-			bHasLaunchedLocal = AutoLaunch(Settings.SettingsVariant.Get<FServiceAutoLaunchSettings>(), *ExecutableInstallPath, HostName, Port);
-			if (bHasLaunchedLocal)
+			const FTimespan MaximumWaitForHealth = FTimespan::FromSeconds(20);
+			FDateTime StartedWaitingForHealth = FDateTime::UtcNow();
+			while (true)
 			{
-				const ZenServerState State(/*ReadOnly*/true);
-				const ZenServerState::ZenServerEntry* RunningEntry = State.LookupByEffectiveListenPort(Port);
-				if (RunningEntry != nullptr)
-				{
-					AutoLaunchedPid = RunningEntry->Pid.load(std::memory_order_relaxed);
-				}
-				AutoLaunchedPort = Port;
-				bIsRunningLocally = true;
-
-				const FTimespan MaximumWaitForHealth = FTimespan::FromSeconds(20);
-				FDateTime StartedWaitingForHealth = FDateTime::UtcNow();
-				bool bLastReadyResult = IsServiceReady();
-				while (!bLastReadyResult)
+				bHasLaunchedLocal = AutoLaunch(Settings.SettingsVariant.Get<FServiceAutoLaunchSettings>(), *ExecutableInstallPath, HostName, Port);
+				if (bHasLaunchedLocal)
 				{
-					FTimespan WaitForHealth = FDateTime::UtcNow() - StartedWaitingForHealth;
-					if (WaitForHealth > MaximumWaitForHealth)
+					const ZenServerState State(/*ReadOnly*/true);
+					const ZenServerState::ZenServerEntry* RunningEntry = State.LookupByEffectiveListenPort(Port);
+					if (RunningEntry != nullptr)
 					{
-						UE_LOG(LogZenServiceInstance, Warning, TEXT("Local ZenServer AutoLaunch initialization timed out waiting for service to become healthy"));
-						break;
+						AutoLaunchedPid = RunningEntry->Pid.load(std::memory_order_relaxed);
 					}
-
-					FPlatformProcess::Sleep(0.5f);
-					if (!IsZenProcessUsingEffectivePort(Port))
+					AutoLaunchedPort = Port;
+					bIsRunningLocally = true;
+					if (IsServiceReady())
 					{
-						AutoLaunch(Settings.SettingsVariant.Get<FServiceAutoLaunchSettings>(), *GetLocalServiceInstallPath(), HostName, Port);
+						break;
 					}
-					bLastReadyResult = IsServiceReady();
 				}
+
+				FTimespan WaitForHealth = FDateTime::UtcNow() - StartedWaitingForHealth;
+				if (WaitForHealth > MaximumWaitForHealth)
+				{
+					bHasLaunchedLocal = false;
+					bIsRunningLocally = false;
+					UE_LOG(LogZenServiceInstance, Warning, TEXT("Local ZenServer AutoLaunch initialization timed out waiting for service to become healthy"));
+					break;
+				}
+				UE_LOG(LogZenServiceInstance, Log, TEXT("Awaiting ZenServer readiness"));
+				FPlatformProcess::Sleep(0.5f);
 			}
 		}
 	}
