Ensure that if zenserver AutoLaunch fails that we have at least one more attempt launch, even if the first attempt exceeded our wait timeout.  This can occur if the first attempt spent a while waiting for the lock to free up.

diff --git a/Engine/Source/Developer/Zen/Private/ZenServerInterface.cpp b/Engine/Source/Developer/Zen/Private/ZenServerInterface.cpp
--- a/Engine/Source/Developer/Zen/Private/ZenServerInterface.cpp
+++ b/Engine/Source/Developer/Zen/Private/ZenServerInterface.cpp
@@ -2116,10 +2116,12 @@ FZenServiceInstance::Initialize()
 		FString ExecutableInstallPath = ConditionalUpdateLocalInstall();
 		if (!ExecutableInstallPath.IsEmpty())
 		{
+			int LaunchAttempts = 0;
 			const FTimespan MaximumWaitForHealth = FTimespan::FromSeconds(20);
 			FDateTime StartedWaitingForHealth = FDateTime::UtcNow();
 			while (true)
 			{
+				++LaunchAttempts;
 				bHasLaunchedLocal = AutoLaunch(Settings.SettingsVariant.Get<FServiceAutoLaunchSettings>(), *ExecutableInstallPath, HostName, Port);
 				if (bHasLaunchedLocal)
 				{
@@ -2138,7 +2140,7 @@ FZenServiceInstance::Initialize()
 				}
 
 				FTimespan WaitForHealth = FDateTime::UtcNow() - StartedWaitingForHealth;
-				if (WaitForHealth > MaximumWaitForHealth)
+				if ((WaitForHealth > MaximumWaitForHealth) && (LaunchAttempts > 1))
 				{
 					bHasLaunchedLocal = false;
 					bIsRunningLocally = false;
