Ensure that zenserver http requests from the zen server interface (not the DDC layer) bypass proxies when communicating with zenserver.

diff --git a/Engine/Source/Developer/Zen/Private/ZenServerHttp.cpp b/Engine/Source/Developer/Zen/Private/ZenServerHttp.cpp
--- a/Engine/Source/Developer/Zen/Private/ZenServerHttp.cpp
+++ b/Engine/Source/Developer/Zen/Private/ZenServerHttp.cpp
@@ -95,6 +95,7 @@ namespace UE::Zen {
 		curl_easy_setopt(Curl, CURLOPT_NOSIGNAL, 1L);
 		curl_easy_setopt(Curl, CURLOPT_DNS_CACHE_TIMEOUT, -1L); // Don't re-resolve names mid-session
 		curl_easy_setopt(Curl, CURLOPT_BUFFERSIZE, 256 * 1024L);
+		curl_easy_setopt(Curl, CURLOPT_NOPROXY, "*");
 		//curl_easy_setopt(Curl, CURLOPT_UPLOAD_BUFFERSIZE, 256 * 1024L);
 		// Response functions
 		curl_easy_setopt(Curl, CURLOPT_HEADERDATA, this);
