Ensure that if we fail to launch zenserver with the CREATE_BREAKAWAY_FROM_JOB flag on Windows that we try again without that flag.

diff --git a/Engine/Source/Developer/Zen/Private/ZenServerInterface.cpp b/Engine/Source/Developer/Zen/Private/ZenServerInterface.cpp
--- a/Engine/Source/Developer/Zen/Private/ZenServerInterface.cpp
+++ b/Engine/Source/Developer/Zen/Private/ZenServerInterface.cpp
@@ -1785,13 +1785,23 @@ StartLocalService(const FZenLocalServiceRunContext& Context, const TCHAR* Transi
 		};
 
 		FString CommandLine = FString::Printf(TEXT("\"%s\" %s"), *PlatformExecutable, *Parms);
-		PROCESS_INFORMATION ProcInfo;
-		if (CreateProcess(NULL, CommandLine.GetCharArray().GetData(), nullptr, nullptr, false, (::DWORD)(NORMAL_PRIORITY_CLASS | DETACHED_PROCESS | CREATE_BREAKAWAY_FROM_JOB), nullptr, PlatformWorkingDirectory.GetCharArray().GetData(), &StartupInfo, &ProcInfo))
+		::DWORD CreationFlagsArray[] = {
+			NORMAL_PRIORITY_CLASS | DETACHED_PROCESS | CREATE_BREAKAWAY_FROM_JOB, // Try with the breakaway flag first
+			NORMAL_PRIORITY_CLASS | DETACHED_PROCESS // If that fails (access denied), try without the breakaway flag next
+		};
+
+		for (::DWORD CreationFlags : CreationFlagsArray)
 		{
-			::CloseHandle(ProcInfo.hThread);
-			Proc = FProcHandle(ProcInfo.hProcess);
+			PROCESS_INFORMATION ProcInfo;
+			if (CreateProcess(NULL, CommandLine.GetCharArray().GetData(), nullptr, nullptr, false, (::DWORD)(NORMAL_PRIORITY_CLASS | DETACHED_PROCESS), nullptr, PlatformWorkingDirectory.GetCharArray().GetData(), &StartupInfo, &ProcInfo))
+			{
+				::CloseHandle(ProcInfo.hThread);
+				Proc = FProcHandle(ProcInfo.hProcess);
+				break;
+			}
 		}
-		else
+
+		if (!Proc.IsValid())
 		{
 			UE_LOG(LogZenServiceInstance, Warning, TEXT("Failed launching %s status: %d."), *CommandLine, GetLastError());
 		}
