Ensure that if we fail to launch zenserver with the CREATE_BREAKAWAY_FROM_JOB flag on Windows that we try again without that flag.

diff --git a/Engine/Source/Developer/Zen/Private/ZenServerInterface.cpp b/Engine/Source/Developer/Zen/Private/ZenServerInterface.cpp
--- a/Engine/Source/Developer/Zen/Private/ZenServerInterface.cpp
+++ b/Engine/Source/Developer/Zen/Private/ZenServerInterface.cpp
@@ -1793,7 +1793,7 @@ StartLocalService(const FZenLocalServiceRunContext& Context, const TCHAR* Transi
 		for (::DWORD CreationFlags : CreationFlagsArray)
 		{
 			PROCESS_INFORMATION ProcInfo;
-			if (CreateProcess(NULL, CommandLine.GetCharArray().GetData(), nullptr, nullptr, false, (::DWORD)(NORMAL_PRIORITY_CLASS | DETACHED_PROCESS), nullptr, PlatformWorkingDirectory.GetCharArray().GetData(), &StartupInfo, &ProcInfo))
+			if (CreateProcess(NULL, CommandLine.GetCharArray().GetData(), nullptr, nullptr, false, CreationFlags, nullptr, PlatformWorkingDirectory.GetCharArray().GetData(), &StartupInfo, &ProcInfo))
 			{
 				::CloseHandle(ProcInfo.hThread);
 				Proc = FProcHandle(ProcInfo.hProcess);
