Don't try to read the payload buffer as a package in FZenCacheStore::FCbPackageReceiver if the request failed or was cancelled
Add hardening to TryLoadCbPackage to handle incomplete messages

diff --git a/Engine/Source/Developer/DerivedDataCache/Private/ZenCacheStore.cpp b/Engine/Source/Developer/DerivedDataCache/Private/ZenCacheStore.cpp
--- a/Engine/Source/Developer/DerivedDataCache/Private/ZenCacheStore.cpp
+++ b/Engine/Source/Developer/DerivedDataCache/Private/ZenCacheStore.cpp
@@ -1486,16 +1486,19 @@ private:
 
 	IHttpReceiver* OnComplete(IHttpResponse& Response) final
 	{
-		FMemoryView MemoryView = MakeMemoryView(BodyArray);
+		if (Response.GetErrorCode() == EHttpErrorCode::None)
 		{
-			FMemoryReaderView Ar(MemoryView);
-			if (Zen::Http::TryLoadCbPackage(Package, Ar))
+			FMemoryView MemoryView = MakeMemoryView(BodyArray);
 			{
-				return Next;
+				FMemoryReaderView Ar(MemoryView);
+				if (Zen::Http::TryLoadCbPackage(Package, Ar))
+				{
+					return Next;
+				}
 			}
+			FMemoryReaderView Ar(MemoryView);
+			Package.TryLoad(Ar);
 		}
-		FMemoryReaderView Ar(MemoryView);
-		Package.TryLoad(Ar);
 		return Next;
 	}
 
diff --git a/Engine/Source/Developer/Zen/Private/ZenSerialization.cpp b/Engine/Source/Developer/Zen/Private/ZenSerialization.cpp
--- a/Engine/Source/Developer/Zen/Private/ZenSerialization.cpp
+++ b/Engine/Source/Developer/Zen/Private/ZenSerialization.cpp
@@ -268,6 +268,10 @@ bool TryLoadCbPackage(FCbPackage& Package, FArchive& Ar, FCbBufferAllocator Allo
 {
 	CbPackageHeader Hdr;
 	Ar.Serialize(&Hdr, sizeof Hdr);
+	if (Ar.IsError())
+	{
+		return false;
+	}
 
 	if (Hdr.HeaderMagic != kCbPkgMagic)
 	{
@@ -278,6 +282,11 @@ bool TryLoadCbPackage(FCbPackage& Package, FArchive& Ar, FCbBufferAllocator Allo
 	AttachmentEntries.SetNum(Hdr.AttachmentCount + 1);
 
 	Ar.Serialize(AttachmentEntries.GetData(), (Hdr.AttachmentCount + 1) * sizeof(CbAttachmentEntry));
+	if (Ar.IsError())
+	{
+		UE_LOG(LogZenSerialization, Warning, TEXT("Package invalid, unable to read %" UINT64_FMT " bytes for attachment header entries"), (Hdr.AttachmentCount + 1) * sizeof(CbAttachmentEntry));
+		return false;
+	}
 
 	int Index = 0;
 
@@ -285,6 +294,11 @@ bool TryLoadCbPackage(FCbPackage& Package, FArchive& Ar, FCbBufferAllocator Allo
 	{
 		FUniqueBuffer AttachmentData = FUniqueBuffer::Alloc(Entry.AttachmentSize);
 		Ar.Serialize(AttachmentData.GetData(), AttachmentData.GetSize());
+		if (Ar.IsError())
+		{
+			UE_LOG(LogZenSerialization, Warning, TEXT("Package invalid, unable to read %" UINT64_FMT " bytes of attachment data"), Entry.AttachmentSize);
+			return false;
+		}
 
 		if (Entry.Flags & CbAttachmentEntry::kIsLocalRef)
 		{
