Make sure we install the PDB for the zen.exe utility when performing an installation and not just the PDB for the zenserver.exe service.

diff --git a/Engine/Source/Developer/Zen/Private/ZenServerInterface.cpp b/Engine/Source/Developer/Zen/Private/ZenServerInterface.cpp
--- a/Engine/Source/Developer/Zen/Private/ZenServerInterface.cpp
+++ b/Engine/Source/Developer/Zen/Private/ZenServerInterface.cpp
@@ -2267,12 +2267,26 @@ FZenServiceInstance::ConditionalUpdateLocalInstall()
 	}
 
 #if PLATFORM_WINDOWS
-	FString InTreeSymbolFilePath = FPaths::SetExtension(InTreeServicePath, TEXT("pdb"));
-	FString InstallSymbolFilePath = FPaths::SetExtension(InstallServicePath, TEXT("pdb"));
-
-	if (FileManager.FileExists(*InTreeSymbolFilePath) && (bMainExecutablesUpdated || !FileManager.FileExists(*InstallSymbolFilePath)))
+	struct FZenExecutable
+	{
+		FString& InTreeFilePath;
+		FString& InstallFilePath;
+	};
+	const FZenExecutable ZenExecutables[] = {
+		// Service executable (zenserver.exe)
+		{InTreeServicePath, InstallServicePath},
+		// Utility executable (zen.exe)
+		{InTreeUtilityPath, InstallUtilityPath},
+	};
+	for (const FZenExecutable& Executable : ZenExecutables)
 	{
-		AttemptFileCopyWithRetries(*InstallSymbolFilePath, *InTreeSymbolFilePath, 1.0);
+		FString InTreeSymbolFilePath = FPaths::SetExtension(Executable.InTreeFilePath, TEXT("pdb"));
+		FString InstallSymbolFilePath = FPaths::SetExtension(Executable.InstallFilePath, TEXT("pdb"));
+
+		if (FileManager.FileExists(*InTreeSymbolFilePath) && (bMainExecutablesUpdated || !FileManager.FileExists(*InstallSymbolFilePath)))
+		{
+			AttemptFileCopyWithRetries(*InstallSymbolFilePath, *InTreeSymbolFilePath, 1.0);
+		}
 	}
 #endif
 
