Add a timeout when locking for zenserver recovery and validate that we get the lock before attempting recovery.
This avoid a race condition where multiple threads would try to start recovery at the same time.

diff --git a/Engine/Source/Developer/Zen/Private/ZenServerInterface.cpp b/Engine/Source/Developer/Zen/Private/ZenServerInterface.cpp
--- a/Engine/Source/Developer/Zen/Private/ZenServerInterface.cpp
+++ b/Engine/Source/Developer/Zen/Private/ZenServerInterface.cpp
@@ -2007,7 +2007,6 @@ FZenServiceInstance::TryRecovery()
 
 	static std::atomic<int64> LastRecoveryTicks;
 	static bool bLastRecoveryResult = false;
-	const FTimespan MaximumWaitForLaunch = FTimespan::FromSeconds(30);
 	const FTimespan MaximumWaitForHealth = FTimespan::FromSeconds(30);
 	const FTimespan MinimumDurationSinceLastRecovery = FTimespan::FromMinutes(2);
 
@@ -2015,7 +2014,12 @@ FZenServiceInstance::TryRecovery()
 
 	if (TimespanSinceLastRecovery > MinimumDurationSinceLastRecovery)
 	{
-		static FSystemWideCriticalSection RecoveryCriticalSection(TEXT("ZenServerRecovery"));
+		static FSystemWideCriticalSection RecoveryCriticalSection(TEXT("ZenServerRecovery"), MaximumWaitForHealth);
+		if (!RecoveryCriticalSection.IsValid())
+		{
+			// A recovery is already in progress but did not complete in time, we assume we failed and let recovery continue on a different thread
+			return false;
+		}
 		// Update timespan since it may have changed since we waited to enter the crit section
 		TimespanSinceLastRecovery = FDateTime::UtcNow() - FDateTime(LastRecoveryTicks.load(std::memory_order_relaxed));
 		if (TimespanSinceLastRecovery > MinimumDurationSinceLastRecovery)
