Verify that a different process has not recovered the zenserver process before forcing a restart

diff --git a/Engine/Source/Developer/Zen/Private/ZenServerInterface.cpp b/Engine/Source/Developer/Zen/Private/ZenServerInterface.cpp
--- a/Engine/Source/Developer/Zen/Private/ZenServerInterface.cpp
+++ b/Engine/Source/Developer/Zen/Private/ZenServerInterface.cpp
@@ -2020,74 +2020,82 @@ FZenServiceInstance::TryRecovery()
 			// A recovery is already in progress but did not complete in time, we assume we failed and let recovery continue on a different thread
 			return false;
 		}
-		// Update timespan since it may have changed since we waited to enter the crit section
-		TimespanSinceLastRecovery = FDateTime::UtcNow() - FDateTime(LastRecoveryTicks.load(std::memory_order_relaxed));
-		if (TimespanSinceLastRecovery > MinimumDurationSinceLastRecovery)
-		{
-			UE_LOG(LogZenServiceInstance, Display, TEXT("Local ZenServer recovery being attempted..."));
 
-			bool bShutdownExistingInstance = true;
-			std::atomic<uint32> PreviousSponsorPids[UE_ARRAY_COUNT(ZenServerState::ZenServerEntry::SponsorPids)];
+		// We test if the service is healthy as a different process might already have triggered a recovery
+		bLastRecoveryResult = IsServiceReady();
+		if (bLastRecoveryResult)
+		{
+			UE_LOG(LogZenServiceInstance, Display, TEXT("Local ZenServer status: Healthy. Skipping recovery"));
+		}
+		else
+		{
+			// Update timespan since it may have changed since we waited to enter the crit section
+			TimespanSinceLastRecovery = FDateTime::UtcNow() - FDateTime(LastRecoveryTicks.load(std::memory_order_relaxed));
+			if (TimespanSinceLastRecovery > MinimumDurationSinceLastRecovery)
 			{
-				const ZenServerState ServerState(/* ReadOnly */true);
-				const ZenServerState::ZenServerEntry* Entry = ServerState.LookupByEffectiveListenPort(Port);
-				if (Entry)
+				UE_LOG(LogZenServiceInstance, Display, TEXT("Local ZenServer recovery being attempted..."));
+
+				bool bShutdownExistingInstance = true;
+				std::atomic<uint32> PreviousSponsorPids[UE_ARRAY_COUNT(ZenServerState::ZenServerEntry::SponsorPids)];
 				{
-					if (Entry->Pid.load(std::memory_order_relaxed) != AutoLaunchedPid)
+					const ZenServerState ServerState(/* ReadOnly */true);
+					const ZenServerState::ZenServerEntry* Entry = ServerState.LookupByEffectiveListenPort(Port);
+					if (Entry)
 					{
-						// The running process pid is not the same as the one we launched.  The process was relaunched elsewhere. Avoid shutting it down again.
-						bShutdownExistingInstance = false;
+						if (Entry->Pid.load(std::memory_order_relaxed) != AutoLaunchedPid)
+						{
+							// The running process pid is not the same as the one we launched.  The process was relaunched elsewhere. Avoid shutting it down again.
+							bShutdownExistingInstance = false;
+						}
+
+						for (int32 SponsorPidIndex = 0; SponsorPidIndex < UE_ARRAY_COUNT(ZenServerState::ZenServerEntry::SponsorPids); ++SponsorPidIndex)
+						{
+							PreviousSponsorPids[SponsorPidIndex].store(Entry->SponsorPids[SponsorPidIndex].load(std::memory_order_relaxed), std::memory_order_relaxed);
+						}
 					}
+				}
+				if (bShutdownExistingInstance && !ShutdownZenServerProcess((int)AutoLaunchedPid))	// !ShutdownRunningServiceUsingEffectivePort(Port))
+				{
+					return false;
+				}
 
-					for (int32 SponsorPidIndex = 0; SponsorPidIndex < UE_ARRAY_COUNT(ZenServerState::ZenServerEntry::SponsorPids); ++SponsorPidIndex)
+				AutoLaunch(Settings.SettingsVariant.Get<FServiceAutoLaunchSettings>(), *GetLocalServiceInstallPath(), HostName, Port);
+				FDateTime StartedWaitingForHealth = FDateTime::UtcNow();
+				bLastRecoveryResult = IsServiceReady();
+				while (!bLastRecoveryResult)
+				{
+					FTimespan WaitForHealth = FDateTime::UtcNow() - StartedWaitingForHealth;
+					if (WaitForHealth > MaximumWaitForHealth)
 					{
-						PreviousSponsorPids[SponsorPidIndex].store(Entry->SponsorPids[SponsorPidIndex].load(std::memory_order_relaxed), std::memory_order_relaxed);
+						UE_LOG(LogZenServiceInstance, Warning, TEXT("Local ZenServer recovery timed out waiting for service to become healthy"));
+						break;
 					}
-				}
-			}
-			if (bShutdownExistingInstance && !ShutdownZenServerProcess((int)AutoLaunchedPid))	// !ShutdownRunningServiceUsingEffectivePort(Port))
-			{
-				return false;
-			}
 
-			AutoLaunch(Settings.SettingsVariant.Get<FServiceAutoLaunchSettings>(), *GetLocalServiceInstallPath(), HostName, Port);
-			FDateTime StartedWaitingForHealth = FDateTime::UtcNow();
-			bLastRecoveryResult = IsServiceReady();
-			while (!bLastRecoveryResult)
-			{
-				FTimespan WaitForHealth = FDateTime::UtcNow() - StartedWaitingForHealth;
-				if (WaitForHealth > MaximumWaitForHealth)
+					FPlatformProcess::Sleep(0.5f);
+					if (!IsZenProcessUsingEffectivePort(Port))
+					{
+						AutoLaunch(Settings.SettingsVariant.Get<FServiceAutoLaunchSettings>(), *GetLocalServiceInstallPath(), HostName, Port);
+					}
+					bLastRecoveryResult = IsServiceReady();
+				}
+				LastRecoveryTicks.store(FDateTime::UtcNow().GetTicks(), std::memory_order_relaxed);
+				UE_LOG(LogZenServiceInstance, Display, TEXT("Local ZenServer recovery finished."));
+				if (bLastRecoveryResult)
 				{
-					UE_LOG(LogZenServiceInstance, Warning, TEXT("Local ZenServer recovery timed out waiting for service to become healthy"));
-					break;
+					UE_LOG(LogZenServiceInstance, Display, TEXT("Local ZenServer post recovery status: Healthy"));
 				}
-
-				FPlatformProcess::Sleep(0.5f);
-				if (!IsZenProcessUsingEffectivePort(Port))
+				else
 				{
-					AutoLaunch(Settings.SettingsVariant.Get<FServiceAutoLaunchSettings>(), *GetLocalServiceInstallPath(), HostName, Port);
+					UE_LOG(LogZenServiceInstance, Display, TEXT("Local ZenServer post recovery status: NOT healthy"));
 				}
-				bLastRecoveryResult = IsServiceReady();
-			}
-			LastRecoveryTicks.store(FDateTime::UtcNow().GetTicks(), std::memory_order_relaxed);
-			UE_LOG(LogZenServiceInstance, Display, TEXT("Local ZenServer recovery finished."));
-			
-			if (bLastRecoveryResult)
-			{
-				UE_LOG(LogZenServiceInstance, Display, TEXT("Local ZenServer post recovery status: Healthy"));
-			}
-			else
-			{
-				UE_LOG(LogZenServiceInstance, Display, TEXT("Local ZenServer post recovery status: NOT healthy"));
-			}
-
-			ZenServerState PostRecoveryServerState(/*bReadOnly*/ false);
-			ZenServerState::ZenServerEntry* PostRecoveryEntry = PostRecoveryServerState.LookupByEffectiveListenPort(Port);
-			if (PostRecoveryEntry)
-			{
-				for (int32 SponsorPidIndex = 0; SponsorPidIndex < UE_ARRAY_COUNT(ZenServerState::ZenServerEntry::SponsorPids); ++SponsorPidIndex)
+				ZenServerState PostRecoveryServerState(/*bReadOnly*/ false);
+				ZenServerState::ZenServerEntry* PostRecoveryEntry = PostRecoveryServerState.LookupByEffectiveListenPort(Port);
+				if (PostRecoveryEntry)
 				{
-					PostRecoveryEntry->SponsorPids[SponsorPidIndex].store(PreviousSponsorPids[SponsorPidIndex].load(std::memory_order_relaxed), std::memory_order_relaxed);
+					for (int32 SponsorPidIndex = 0; SponsorPidIndex < UE_ARRAY_COUNT(ZenServerState::ZenServerEntry::SponsorPids); ++SponsorPidIndex)
+					{
+						PostRecoveryEntry->SponsorPids[SponsorPidIndex].store(PreviousSponsorPids[SponsorPidIndex].load(std::memory_order_relaxed), std::memory_order_relaxed);
+					}
 				}
 			}
 		}
