Separate the control mechanism for launching zenserver with public network interfaces into its own config parameter instead of keeping it in the catch-all ExtraArgs config parameter using the --http-forceloopback argument.

diff --git a/Engine/Config/BaseEngine.ini b/Engine/Config/BaseEngine.ini
--- a/Engine/Config/BaseEngine.ini
+++ b/Engine/Config/BaseEngine.ini
@@ -2504,10 +2504,11 @@ AutoLaunch=true
 [Zen.AutoLaunch]
 DataPath=%ENGINEVERSIONAGNOSTICINSTALLEDUSERDIR%Zen/Data
 LimitProcessLifetime=true
+AllowPublicNetworkInterface=false
 LocalDataCachePathEnvOverride=UE-LocalDataCachePath
 LocalDataCachePathCommandLineOverride=LocalDataCachePath
 LocalDataCachePathEditorOverrideSetting=LocalDerivedDataCache
-ExtraArgs=--http asio --gc-cache-duration-seconds 1209600 --gc-interval-seconds 21600 --gc-low-diskspace-threshold 2147483648 --http-forceloopback --quiet
+ExtraArgs=--http asio --gc-cache-duration-seconds 1209600 --gc-interval-seconds 21600 --gc-low-diskspace-threshold 2147483648 --quiet
 
 [Zen.ConnectExisting]
 HostName=localhost
diff --git a/Engine/Source/Developer/Zen/Private/ZenServerInterface.cpp b/Engine/Source/Developer/Zen/Private/ZenServerInterface.cpp
--- a/Engine/Source/Developer/Zen/Private/ZenServerInterface.cpp
+++ b/Engine/Source/Developer/Zen/Private/ZenServerInterface.cpp
@@ -946,6 +946,7 @@ FServiceSettings::ReadFromConfig()
 			GConfig->GetBool(AutoLaunchConfigSection, TEXT("ShowConsole"), AutoLaunchSettings.bShowConsole, GEngineIni);
 			GConfig->GetBool(AutoLaunchConfigSection, TEXT("LimitProcessLifetime"), AutoLaunchSettings.bLimitProcessLifetime, GEngineIni);
 			ApplyProcessLifetimeOverride(AutoLaunchSettings.bLimitProcessLifetime);
+			GConfig->GetBool(AutoLaunchConfigSection, TEXT("AllowPublicNetworkInterface"), AutoLaunchSettings.bAllowPublicNetworkInterface, GEngineIni);
 			EnsureEditorSettingsConfigLoaded();
 			GConfig->GetBool(TEXT("/Script/UnrealEd.CrashReportsPrivacySettings"), TEXT("bSendUnattendedBugReports"), AutoLaunchSettings.bSendUnattendedBugReports, GEditorSettingsIni);
 		}
@@ -982,6 +983,7 @@ FServiceSettings::ReadFromCompactBinary(FCbFieldView Field)
 				AutoLaunchSettings.bIsDefaultDataPath = AutoLaunchSettingsObject["IsDefaultDataPath"].AsBool();
 				AutoLaunchSettings.bLimitProcessLifetime = AutoLaunchSettingsObject["LimitProcessLifetime"].AsBool();
 				ApplyProcessLifetimeOverride(AutoLaunchSettings.bLimitProcessLifetime);
+				AutoLaunchSettings.bAllowPublicNetworkInterface = AutoLaunchSettingsObject["AllowPublicNetworkInterface"].AsBool();
 				AutoLaunchSettings.bSendUnattendedBugReports = AutoLaunchSettingsObject["SendUnattendedBugReports"].AsBool();
 				AutoLaunchSettings.bIsDefaultSharedRunContext = AutoLaunchSettingsObject["IsDefaultSharedRunContext"].AsBool(AutoLaunchSettings.bIsDefaultSharedRunContext);
 			}
@@ -1042,6 +1044,7 @@ FServiceSettings::WriteToCompactBinary(FCbWriter& Writer) const
 		Writer << "ShowConsole" << AutoLaunchSettings.bShowConsole;
 		Writer << "IsDefaultDataPath" << AutoLaunchSettings.bIsDefaultDataPath;
 		Writer << "LimitProcessLifetime" << AutoLaunchSettings.bLimitProcessLifetime;
+		Writer << "AllowPublicNetworkInterface" << AutoLaunchSettings.bAllowPublicNetworkInterface;
 		Writer << "SendUnattendedBugReports" << AutoLaunchSettings.bSendUnattendedBugReports;
 		Writer << "IsDefaultSharedRunContext" << AutoLaunchSettings.bIsDefaultSharedRunContext;
 		Writer.EndObject();
@@ -1639,6 +1642,11 @@ DetermineCmdLineWithoutTransientComponents(const FServiceAutoLaunchSettings& InS
 		Parms.Append(TEXT(" --no-sentry"));
 	}
 
+	if (!InSettings.bAllowPublicNetworkInterface)
+	{
+		Parms.Append(TEXT(" --http-forceloopback"));
+	}
+
 	return Parms;
 }
 
diff --git a/Engine/Source/Developer/Zen/Public/Experimental/ZenServerInterface.h b/Engine/Source/Developer/Zen/Public/Experimental/ZenServerInterface.h
--- a/Engine/Source/Developer/Zen/Public/Experimental/ZenServerInterface.h
+++ b/Engine/Source/Developer/Zen/Public/Experimental/ZenServerInterface.h
@@ -44,6 +44,7 @@ struct FServiceAutoLaunchSettings
 	bool bShowConsole = false;
 	bool bIsDefaultDataPath = false;
 	bool bLimitProcessLifetime = false;
+	bool bAllowPublicNetworkInterface = false;
 	bool bSendUnattendedBugReports = false;
 	bool bIsDefaultSharedRunContext = true;
 };
