Rename Zen.AutoLaunch parameter AllowPublicNetworkInterface to AllowRemoteNetworkService ahead of 5.5 release so that terminology is clearer to understand and document (avoids confusion around public/private network terms as used in Windows).

diff --git a/Engine/Config/BaseEngine.ini b/Engine/Config/BaseEngine.ini
--- a/Engine/Config/BaseEngine.ini
+++ b/Engine/Config/BaseEngine.ini
@@ -2557,7 +2557,7 @@ AutoLaunch=true
 [Zen.AutoLaunch]
 DataPath=%ENGINEVERSIONAGNOSTICINSTALLEDUSERDIR%Zen/Data
 LimitProcessLifetime=true
-AllowPublicNetworkInterface=false
+AllowRemoteNetworkService=false
 LocalDataCachePathEnvOverride=UE-LocalDataCachePath
 LocalDataCachePathCommandLineOverride=LocalDataCachePath
 LocalDataCachePathEditorOverrideSetting=LocalDerivedDataCache
diff --git a/Engine/Source/Developer/Zen/Private/ZenServerInterface.cpp b/Engine/Source/Developer/Zen/Private/ZenServerInterface.cpp
--- a/Engine/Source/Developer/Zen/Private/ZenServerInterface.cpp
+++ b/Engine/Source/Developer/Zen/Private/ZenServerInterface.cpp
@@ -946,7 +946,7 @@ FServiceSettings::ReadFromConfig()
 			GConfig->GetBool(AutoLaunchConfigSection, TEXT("ShowConsole"), AutoLaunchSettings.bShowConsole, GEngineIni);
 			GConfig->GetBool(AutoLaunchConfigSection, TEXT("LimitProcessLifetime"), AutoLaunchSettings.bLimitProcessLifetime, GEngineIni);
 			ApplyProcessLifetimeOverride(AutoLaunchSettings.bLimitProcessLifetime);
-			GConfig->GetBool(AutoLaunchConfigSection, TEXT("AllowPublicNetworkInterface"), AutoLaunchSettings.bAllowPublicNetworkInterface, GEngineIni);
+			GConfig->GetBool(AutoLaunchConfigSection, TEXT("AllowRemoteNetworkService"), AutoLaunchSettings.bAllowRemoteNetworkService, GEngineIni);
 			EnsureEditorSettingsConfigLoaded();
 			GConfig->GetBool(TEXT("/Script/UnrealEd.CrashReportsPrivacySettings"), TEXT("bSendUnattendedBugReports"), AutoLaunchSettings.bSendUnattendedBugReports, GEditorSettingsIni);
 		}
@@ -983,7 +983,7 @@ FServiceSettings::ReadFromCompactBinary(FCbFieldView Field)
 				AutoLaunchSettings.bIsDefaultDataPath = AutoLaunchSettingsObject["IsDefaultDataPath"].AsBool();
 				AutoLaunchSettings.bLimitProcessLifetime = AutoLaunchSettingsObject["LimitProcessLifetime"].AsBool();
 				ApplyProcessLifetimeOverride(AutoLaunchSettings.bLimitProcessLifetime);
-				AutoLaunchSettings.bAllowPublicNetworkInterface = AutoLaunchSettingsObject["AllowPublicNetworkInterface"].AsBool();
+				AutoLaunchSettings.bAllowRemoteNetworkService = AutoLaunchSettingsObject["AllowRemoteNetworkService"].AsBool();
 				AutoLaunchSettings.bSendUnattendedBugReports = AutoLaunchSettingsObject["SendUnattendedBugReports"].AsBool();
 				AutoLaunchSettings.bIsDefaultSharedRunContext = AutoLaunchSettingsObject["IsDefaultSharedRunContext"].AsBool(AutoLaunchSettings.bIsDefaultSharedRunContext);
 			}
@@ -1044,7 +1044,7 @@ FServiceSettings::WriteToCompactBinary(FCbWriter& Writer) const
 		Writer << "ShowConsole" << AutoLaunchSettings.bShowConsole;
 		Writer << "IsDefaultDataPath" << AutoLaunchSettings.bIsDefaultDataPath;
 		Writer << "LimitProcessLifetime" << AutoLaunchSettings.bLimitProcessLifetime;
-		Writer << "AllowPublicNetworkInterface" << AutoLaunchSettings.bAllowPublicNetworkInterface;
+		Writer << "AllowRemoteNetworkService" << AutoLaunchSettings.bAllowRemoteNetworkService;
 		Writer << "SendUnattendedBugReports" << AutoLaunchSettings.bSendUnattendedBugReports;
 		Writer << "IsDefaultSharedRunContext" << AutoLaunchSettings.bIsDefaultSharedRunContext;
 		Writer.EndObject();
@@ -1642,7 +1642,7 @@ DetermineCmdLineWithoutTransientComponents(const FServiceAutoLaunchSettings& InS
 		Parms.Append(TEXT(" --no-sentry"));
 	}
 
-	if (!InSettings.bAllowPublicNetworkInterface)
+	if (!InSettings.bAllowRemoteNetworkService)
 	{
 		Parms.Append(TEXT(" --http-forceloopback"));
 	}
diff --git a/Engine/Source/Developer/Zen/Public/Experimental/ZenServerInterface.h b/Engine/Source/Developer/Zen/Public/Experimental/ZenServerInterface.h
--- a/Engine/Source/Developer/Zen/Public/Experimental/ZenServerInterface.h
+++ b/Engine/Source/Developer/Zen/Public/Experimental/ZenServerInterface.h
@@ -44,7 +44,7 @@ struct FServiceAutoLaunchSettings
 	bool bShowConsole = false;
 	bool bIsDefaultDataPath = false;
 	bool bLimitProcessLifetime = false;
-	bool bAllowPublicNetworkInterface = false;
+	bool bAllowRemoteNetworkService = false;
 	bool bSendUnattendedBugReports = false;
 	bool bIsDefaultSharedRunContext = true;
 };
