Use named event to wait for zenserver readyness instead of agressively polling the zenserver lock file
Remove the restoring of sponsored processes when trying to recover - we should never touch them unless we are adding a process ourselves.
Refactored out ZenServerState and ZenNamedEvent into separate file so we don't have low-level zenserver specific code in ZenServerInterface.cpp.

diff --git a/Engine/Source/Developer/Zen/Private/ZenServerInterface.cpp b/Engine/Source/Developer/Zen/Private/ZenServerInterface.cpp
--- a/Engine/Source/Developer/Zen/Private/ZenServerInterface.cpp
+++ b/Engine/Source/Developer/Zen/Private/ZenServerInterface.cpp
@@ -5,6 +5,8 @@
 #include "ZenBackendUtils.h"
 #include "ZenSerialization.h"
 #include "ZenServerHttp.h"
+#include "ZenServerState.h"
+#include "ZenVersion.h"
 
 #include "AnalyticsEventAttribute.h"
 #include "Async/Async.h"
@@ -27,28 +29,18 @@
 #include "Misc/ScopedSlowTask.h"
 #include "Misc/ScopeExit.h"
 #include "Serialization/CompactBinarySerialization.h"
-#include "Serialization/CompactBinaryValidation.h"
 #include "Serialization/JsonReader.h"
 #include "Serialization/JsonSerializer.h"
 #include "Serialization/JsonWriter.h"
 #include "String/LexFromString.h"
-#include "ZenVersion.h"
 #include "Serialization/CompactBinaryWriter.h"
 
 #if PLATFORM_WINDOWS
 #	include "Windows/AllowWindowsPlatformTypes.h"
-#	include <shellapi.h>
-#	include <synchapi.h>
 #	include <Windows.h>
 #	include "Windows/HideWindowsPlatformTypes.h"
 #endif
 
-#if PLATFORM_UNIX || PLATFORM_MAC
-#	include <sys/file.h>
-#	include <sys/mman.h>
-#	include <sys/sem.h>
-#endif
-
 #define ALLOW_SETTINGS_OVERRIDE_FROM_COMMANDLINE			(UE_SERVER || !(UE_BUILD_SHIPPING))
 
 namespace UE::Zen
@@ -56,356 +48,6 @@ namespace UE::Zen
 
 DEFINE_LOG_CATEGORY_STATIC(LogZenServiceInstance, Log, All);
 
-// Native functions to interact with a process using a process id
-// We don't use UE's own OpenProcess as they try to open processes with PROCESS_ALL_ACCESS
-static bool NativeIsProcessRunning(uint32 Pid)
-{
-	if (Pid == 0)
-	{
-		return false;
-	}
-#if PLATFORM_WINDOWS
-	HANDLE Handle = OpenProcess(PROCESS_QUERY_LIMITED_INFORMATION, 0, (DWORD)Pid);
-
-	if (!Handle)
-	{
-		DWORD Error = GetLastError();
-
-		if (Error == ERROR_INVALID_PARAMETER)
-		{
-			return false;
-		}
-		else if (Error == ERROR_ACCESS_DENIED)
-		{
-			UE_LOG(LogZenServiceInstance, Warning, TEXT("No access to open running process %d: %d, assuming it is running"), Pid, Error);
-			return true;
-		}
-		UE_LOG(LogZenServiceInstance, Warning, TEXT("Failed to open running process %d: %d, assuming it is not running"), Pid, Error);
-		return false;
-	}
-	ON_SCOPE_EXIT{ CloseHandle(Handle); };
-
-	DWORD ExitCode = 0;
-	if (GetExitCodeProcess(Handle, &ExitCode) == 0)
-	{
-		DWORD Error = GetLastError();
-		UE_LOG(LogZenServiceInstance, Warning, TEXT("Failed to get running process exit code %d: %d, assuming it is still running"), Pid, Error);
-		return true;
-	}
-	else if (ExitCode == STILL_ACTIVE)
-	{
-		return true;
-	}
-
-	return false;
-
-#elif PLATFORM_UNIX || PLATFORM_MAC
-	int Res = kill(pid_t(Pid), 0);
-	if (Res == 0)
-	{
-		return true;
-	}
-	int Error = errno;
-	if (Error == EPERM)
-	{
-		UE_LOG(LogZenServiceInstance, Warning, TEXT("No permission to signal running process %d: %d, assuming it is running"), Pid, Error);
-		return true;
-	}
-	else if (Error == ESRCH)
-	{
-		return false;
-	}
-	UE_LOG(LogZenServiceInstance, Warning, TEXT("Failed to signal running process %d: %d, assuming it is running"), Pid, Error);
-	return true;
-#endif
-}
-
-static bool NativeTerminate(uint32 Pid)
-{
-#if PLATFORM_WINDOWS
-	HANDLE Handle = OpenProcess(PROCESS_TERMINATE | SYNCHRONIZE, 0, (DWORD)Pid);
-	if (Handle == NULL)
-	{
-		DWORD Error = GetLastError();
-
-		if (Error != ERROR_INVALID_PARAMETER)
-		{
-			UE_LOG(LogZenServiceInstance, Warning, TEXT("Failed to open running process for terminate %d: %d"), Pid, Error);
-			return false;
-		}
-		return true;
-	}
-	ON_SCOPE_EXIT{ CloseHandle(Handle); };
-
-	BOOL bTerminated = TerminateProcess(Handle, 0);
-	if (!bTerminated)
-	{
-		DWORD Error = GetLastError();
-		UE_LOG(LogZenServiceInstance, Warning, TEXT("Failed to terminate running process %d: %d"), Pid, Error);
-		return false;
-	}
-	DWORD WaitResult = WaitForSingleObject(Handle, 15000);
-	BOOL bSuccess = (WaitResult == WAIT_OBJECT_0) || (WaitResult == WAIT_ABANDONED_0);
-	if (!bSuccess)
-	{
-		if (WaitResult == WAIT_FAILED)
-		{
-			DWORD Error = GetLastError();
-			UE_LOG(LogZenServiceInstance, Warning, TEXT("Failed to wait for terminated process %d: %d"), Pid, Error);
-		}
-		return false;
-	}
-#elif PLATFORM_UNIX || PLATFORM_MAC
-	int Res = kill(pid_t(Pid), SIGKILL);
-	if (Res != 0)
-	{
-		int err = errno;
-		if (err != ESRCH)
-		{
-			UE_LOG(LogZenServiceInstance, Warning, TEXT("Failed to terminate running process %d: %d"), Pid, err);
-			return false;
-		}
-	}
-#endif
-	return true;
-}
-
-class ZenServerState
-{
-public:
-	ZenServerState(bool ReadOnly);
-	~ZenServerState();
-
-	struct ZenServerEntry
-	{
-		// This matches the structure found in the Zen server
-		// https://github.com/EpicGames/zen/blob/main/zenutil/include/zenutil/zenserverprocess.h#L91
-		//
-		std::atomic<uint32> Pid;
-		std::atomic<uint16> DesiredListenPort;
-		std::atomic<uint16> Flags;
-		uint8				SessionId[12];
-		std::atomic<uint32> SponsorPids[8];
-		std::atomic<uint16> EffectiveListenPort;
-		uint8				Padding[10];
-
-		enum class FlagsEnum : uint16
-		{
-			kShutdownPlease = 1 << 0,
-			kIsReady = 1 << 1,
-		};
-
-		bool AddSponsorProcess(uint32 PidToAdd);
-	};
-	static_assert(sizeof(ZenServerEntry) == 64);
-
-	const ZenServerEntry* LookupByDesiredListenPort(int DesiredListenPort) const;
-	ZenServerEntry* LookupByDesiredListenPort(int DesiredListenPort);
-	const ZenServerEntry* LookupByEffectiveListenPort(int EffectiveListenPort) const;
-	ZenServerEntry* LookupByEffectiveListenPort(int EffectiveListenPort);
-	const ZenServerEntry* LookupByPid(uint32 Pid) const;
-
-private:
-	const ZenServerEntry* LookupByDesiredListenPortInternal(int DesiredListenPort) const;
-	const ZenServerEntry* LookupByEffectiveListenPortInternal(int EffectiveListenPort) const;
-	void* m_hMapFile = nullptr;
-	ZenServerEntry* m_Data = nullptr;
-	int				m_MaxEntryCount = 65536 / sizeof(ZenServerEntry);
-	bool			m_IsReadOnly = true;
-};
-
-ZenServerState::ZenServerState(bool ReadOnly)
-	: m_hMapFile(nullptr)
-	, m_Data(nullptr)
-{
-	size_t MapSize = m_MaxEntryCount * sizeof(ZenServerEntry);
-
-#if PLATFORM_WINDOWS
-	DWORD DesiredAccess = ReadOnly ? FILE_MAP_READ : (FILE_MAP_READ | FILE_MAP_WRITE);
-	HANDLE hMap = OpenFileMapping(DesiredAccess, 0, L"Global\\ZenMap");
-	if (hMap == NULL)
-	{
-		hMap = OpenFileMapping(DesiredAccess, 0, L"Local\\ZenMap");
-	}
-
-	if (hMap == NULL)
-	{
-		return;
-	}
-
-	void* pBuf = MapViewOfFile(hMap,		   // handle to map object
-		DesiredAccess,  // read permission
-		0,			   // offset high
-		0,			   // offset low
-		MapSize);
-
-	if (pBuf == NULL)
-	{
-		CloseHandle(hMap);
-		return;
-	}
-#elif PLATFORM_UNIX || PLATFORM_MAC
-	int OFlag = ReadOnly ? (O_RDONLY | O_CLOEXEC) : (O_RDWR | O_CREAT | O_CLOEXEC);
-	int Fd = shm_open("/UnrealEngineZen", OFlag, 0666);
-	if (Fd < 0)
-	{
-		return;
-	}
-	void* hMap = (void*)intptr_t(Fd);
-
-	int Prot = ReadOnly ? PROT_READ : (PROT_WRITE | PROT_READ);
-	void* pBuf = mmap(nullptr, MapSize, Prot, MAP_SHARED, Fd, 0);
-	if (pBuf == MAP_FAILED)
-	{
-		close(Fd);
-		return;
-	}
-#endif
-
-#if PLATFORM_WINDOWS || PLATFORM_UNIX || PLATFORM_MAC
-	m_hMapFile = hMap;
-	m_Data = reinterpret_cast<ZenServerEntry*>(pBuf);
-#endif
-	m_IsReadOnly = ReadOnly;
-}
-
-ZenServerState::~ZenServerState()
-{
-#if PLATFORM_WINDOWS
-	if (m_Data)
-	{
-		UnmapViewOfFile(m_Data);
-	}
-
-	if (m_hMapFile)
-	{
-		CloseHandle(m_hMapFile);
-	}
-#elif PLATFORM_UNIX || PLATFORM_MAC
-	if (m_Data != nullptr)
-	{
-		munmap((void*)m_Data, m_MaxEntryCount * sizeof(ZenServerEntry));
-	}
-
-	int Fd = int(intptr_t(m_hMapFile));
-	close(Fd);
-#endif
-	m_hMapFile = nullptr;
-	m_Data = nullptr;
-}
-
-const ZenServerState::ZenServerEntry* ZenServerState::LookupByDesiredListenPortInternal(int Port) const
-{
-	if (m_Data == nullptr)
-	{
-		return nullptr;
-	}
-
-	for (int i = 0; i < m_MaxEntryCount; ++i)
-	{
-		if (m_Data[i].DesiredListenPort.load(std::memory_order_relaxed) == Port)
-		{
-			const ZenServerState::ZenServerEntry* Entry = &m_Data[i];
-			if (NativeIsProcessRunning((uint32)Entry->Pid.load(std::memory_order_relaxed)))
-			{
-				return Entry;
-			}
-		}
-	}
-
-	return nullptr;
-}
-
-const ZenServerState::ZenServerEntry* ZenServerState::LookupByDesiredListenPort(int Port) const
-{
-	return LookupByDesiredListenPortInternal(Port);
-}
-
-ZenServerState::ZenServerEntry* ZenServerState::LookupByDesiredListenPort(int Port)
-{
-	check(!m_IsReadOnly);
-	return const_cast<ZenServerState::ZenServerEntry*>(LookupByDesiredListenPortInternal(Port));
-}
-
-const ZenServerState::ZenServerEntry* ZenServerState::LookupByEffectiveListenPortInternal(int Port) const
-{
-	if (m_Data == nullptr)
-	{
-		return nullptr;
-	}
-
-	for (int i = 0; i < m_MaxEntryCount; ++i)
-	{
-		const ZenServerState::ZenServerEntry* Entry = &m_Data[i];
-		if (Entry->EffectiveListenPort.load(std::memory_order_relaxed) == Port)
-		{
-			if (NativeIsProcessRunning((uint32)Entry->Pid.load(std::memory_order_relaxed)))
-			{
-				return Entry;
-			}
-		}
-	}
-
-	return nullptr;
-}
-
-const ZenServerState::ZenServerEntry* ZenServerState::LookupByEffectiveListenPort(int Port) const
-{
-	return LookupByEffectiveListenPortInternal(Port);
-}
-
-ZenServerState::ZenServerEntry* ZenServerState::LookupByEffectiveListenPort(int Port)
-{
-	check(!m_IsReadOnly);
-	return const_cast<ZenServerState::ZenServerEntry*>(LookupByEffectiveListenPortInternal(Port));
-}
-
-const ZenServerState::ZenServerEntry* ZenServerState::LookupByPid(uint32 Pid) const
-{
-	if (m_Data == nullptr)
-	{
-		return nullptr;
-	}
-
-	for (int i = 0; i < m_MaxEntryCount; ++i)
-	{
-		const ZenServerState::ZenServerEntry* Entry = &m_Data[i];
-		if (m_Data[i].Pid.load(std::memory_order_relaxed) == Pid)
-		{
-			if (NativeIsProcessRunning(Pid))
-			{
-				return Entry;
-			}
-		}
-	}
-
-	return nullptr;
-}
-
-bool
-ZenServerState::ZenServerEntry::AddSponsorProcess(uint32 PidToAdd)
-{
-	for (std::atomic<uint32>& PidEntry : SponsorPids)
-	{
-		if (PidEntry.load(std::memory_order_relaxed) == 0)
-		{
-			uint32 Expected = 0;
-			if (PidEntry.compare_exchange_strong(Expected, PidToAdd))
-			{
-				// Success!
-				return true;
-			}
-		}
-		else if (PidEntry.load(std::memory_order_relaxed) == PidToAdd)
-		{
-			// Success, the because pid is already in the list
-			return true;
-		}
-	}
-
-	return false;
-}
-
 static FZenVersion
 GetZenVersion(const FString& UtilityPath, const FString& ServicePath, const FString& VersionCachePath)
 {
@@ -1101,117 +743,6 @@ FServiceSettings::TryApplyAutoLaunchOverride()
 uint16 FZenServiceInstance::AutoLaunchedPort = 0;
 uint32 FZenServiceInstance::AutoLaunchedPid = 0;
 
-struct LockFileData
-{
-	uint32 ProcessId = 0;
-	FString DataDir;
-	uint16 EffectivePort = 0;
-	FCbObjectId SessionId;
-	bool IsReady = false;
-	bool IsValid = false;
-};
-
-static LockFileData ReadLockData(FUniqueBuffer&& FileBytes)
-{
-	if (ValidateCompactBinary(FileBytes, ECbValidateMode::Default) == ECbValidateError::None)
-	{
-		FCbObject LockObject(FileBytes.MoveToShared());
-		int32 ProcessId = LockObject["pid"].AsInt32();
-		FUtf8StringView DataDir = LockObject["data"].AsString();
-		int32 EffectivePort = LockObject["port"].AsInt32();
-		FCbObjectId SessionId = LockObject["session_id"].AsObjectId();
-		bool IsReady = LockObject["ready"].AsBool();
-		bool IsValid = ProcessId > 0 && EffectivePort > 0 && EffectivePort <= 0xffff;
-		return LockFileData{static_cast<uint32>(ProcessId), FString(DataDir), static_cast<uint16>(EffectivePort), SessionId, IsReady, IsValid};
-	}
-	return LockFileData{ 0, {}, 0, {}, false};
-}
-
-static LockFileData
-ReadCbLockFile(FStringView FileName)
-{
-#if PLATFORM_WINDOWS
-	// Windows specific lock reading path
-	// Uses share flags that are unique to windows to allow us to read file contents while the file may be open for write AND delete by another process (zenserver).
-
-	uint32 Access = GENERIC_READ;
-	uint32 WinFlags = FILE_SHARE_READ | FILE_SHARE_WRITE | FILE_SHARE_DELETE;
-	uint32 Create = OPEN_EXISTING;
-
-	TStringBuilder<MAX_PATH> FullFileNameBuilder;
-	FPathViews::ToAbsolutePath(FileName, FullFileNameBuilder);
-	for (TCHAR& Char : MakeArrayView(FullFileNameBuilder))
-	{
-		if (Char == TEXT('/'))
-		{
-			Char = TEXT('\\');
-		}
-	}
-	if (FullFileNameBuilder.Len() >= MAX_PATH)
-	{
-		FullFileNameBuilder.Prepend(TEXTVIEW("\\\\?\\"));
-	}
-	HANDLE Handle = CreateFileW(FullFileNameBuilder.ToString(), Access, WinFlags, NULL, Create, FILE_ATTRIBUTE_NORMAL, NULL);
-	if (Handle != INVALID_HANDLE_VALUE)
-	{
-		ON_SCOPE_EXIT { CloseHandle(Handle); };
-		LARGE_INTEGER LI;
-		if (GetFileSizeEx(Handle, &LI))
-		{
-			checkf(LI.QuadPart == LI.u.LowPart, TEXT("Lock file exceeds supported 2GB limit."));
-			int32 FileSize32 = LI.u.LowPart;
-			FUniqueBuffer FileBytes = FUniqueBuffer::Alloc(FileSize32);
-			DWORD ReadBytes = 0;
-			if (ReadFile(Handle, FileBytes.GetData(), FileSize32, &ReadBytes, NULL) && (ReadBytes == FileSize32))
-			{
-				return ReadLockData(std::move(FileBytes));
-			}
-		}
-	}
-	return {};
-#elif PLATFORM_UNIX || PLATFORM_MAC
-	TAnsiStringBuilder<256> LockFilePath;
-	LockFilePath << FileName;
-	int32 Fd = open(LockFilePath.ToString(), O_RDONLY);
-	if (Fd < 0)
-	{
-		return {};
-	}
-
-	// If we can claim the lock then it's an orphaned lock file and should be
-	// ignored. Not ideal as there's a period of time when the lock can be
-	// held unncessarily.
-	int32 LockRet = flock(Fd, LOCK_EX | LOCK_NB);
-	if (LockRet >= 0)
-	{
-		unlink(LockFilePath.ToString());
-		flock(Fd, LOCK_UN);
-		close(Fd);
-		return {};
-	}
-
-	if (errno != EWOULDBLOCK && errno != EAGAIN)
-	{
-		close(Fd);
-		return {};
-	}
-
-	struct stat Stat;
-	fstat(Fd, &Stat);
-	uint64 FileSize = uint64(Stat.st_size);
-
-	FUniqueBuffer FileBytes = FUniqueBuffer::Alloc(FileSize);
-	if (read(Fd, FileBytes.GetData(), FileSize) == FileSize)
-	{
-		close(Fd);
-		return ReadLockData(std::move(FileBytes));
-	}
-
-	close(Fd);
-	return {};
-#endif
-}
-
 FString GetLocalServiceExecutableName()
 {
 	return
@@ -1224,192 +755,49 @@ FString GetLocalServiceExecutableName()
 }
 
 FString
-GetLocalServiceInstallPath()
+GetLocalInstallPath()
 {
-	return FPaths::ConvertRelativePathToFull(FPaths::Combine(GetLocalZenRootPath(), TEXT("Zen\\Install"), GetLocalServiceExecutableName()));
+	return FPaths::ConvertRelativePathToFull(FPaths::Combine(GetLocalZenRootPath(), TEXT("Zen\\Install")));
 }
 
-static bool
-IsLockFileLocked(const TCHAR* FileName, bool bAttemptCleanUp = false)
-{
-#if PLATFORM_WINDOWS
-	if (bAttemptCleanUp)
-	{
-		IFileManager::Get().Delete(FileName, false, false, true);
-	}
-	return IFileManager::Get().FileExists(FileName);
-#elif PLATFORM_UNIX || PLATFORM_MAC
-	TAnsiStringBuilder<256> LockFilePath;
-	LockFilePath << FileName;
-	int32 Fd = open(LockFilePath.ToString(), O_RDONLY);
-	if (Fd < 0)
-	{
-		return false;
-	}
-
-	int32 LockRet = flock(Fd, LOCK_EX | LOCK_NB);
-	if (LockRet < 0)
-	{
-		close(Fd);
-		return errno == EWOULDBLOCK || errno == EAGAIN;
-	}
-
-	// Consider the lock file as orphaned if we we managed to claim the lock for
-	// it. Might as well delete it while we own it.
-	unlink(LockFilePath.ToString());
-
-	flock(Fd, LOCK_UN);
-	close(Fd);
-
-	return false;
-#endif
-}
-
-#if PLATFORM_WINDOWS
-static FString GetZenProcessShutdownEventName(uint16 EffectiveListenPort)
+FString
+GetLocalServiceInstallPath()
 {
-	return *WriteToWideString<64>(WIDETEXT("Zen_"), EffectiveListenPort, WIDETEXT("_Shutdown"));
+	return FPaths::ConvertRelativePathToFull(FPaths::Combine(GetLocalInstallPath(), GetLocalServiceExecutableName()));
 }
 
-static HANDLE OpenNativeEvent(const FString& EventName, DWORD Access)
-{
-	return OpenEventW(Access, false, *EventName);
-}
-#elif PLATFORM_UNIX || PLATFORM_MAC
-static FAnsiString GetZenProcessShutdownEventName(uint16 EffectiveListenPort)
-{
-	return *WriteToAnsiString<64>("/tmp/Zen_", EffectiveListenPort, "_Shutdown");
-}
-static int OpenNativeEvent(const FAnsiString& EventName, int Access)
+FString
+GetLocalInstallUtilityPath()
 {
-	key_t IpcKey = ftok(*EventName, 1);
-	if (IpcKey < 0)
-	{
-		return -1;
-	}
-	int Semaphore = semget(IpcKey, 1, Access);
-	if (Semaphore < 0)
-	{
-		return -1;
-	}
-	return Semaphore;
-}
+	return FPaths::ConvertRelativePathToFull(FPaths::Combine(GetLocalInstallPath(),
+#if PLATFORM_WINDOWS
+		TEXT("zen.exe")
+#else
+		TEXT("zen")
 #endif
+	));
+}
 
 static bool
 IsZenProcessUsingEffectivePort(uint16 EffectiveListenPort)
 {
-#if PLATFORM_WINDOWS
-	HANDLE Handle = OpenNativeEvent(GetZenProcessShutdownEventName(EffectiveListenPort), READ_CONTROL);
-	if (Handle != NULL)
-	{
-		CloseHandle(Handle);
-		return true;
-	}
-	return false;
-#elif PLATFORM_UNIX || PLATFORM_MAC
-	int Semaphore = OpenNativeEvent(GetZenProcessShutdownEventName(EffectiveListenPort), 0400);
-	if (Semaphore >= 0)
-	{
-		return true;
-	}
-	return false;
-#else
-	static_assert(false, "Missing implementation for Zen named shutdown events");
-	return false;
-#endif
+	ZenSharedEvent ShutDownEvent(ZenSharedEvent::GetShutdownEventName(EffectiveListenPort));
+	return ShutDownEvent.Exists();
 }
 
 static bool
 RequestZenShutdownOnEffectivePort(uint16 EffectiveListenPort)
 {
-#if PLATFORM_WINDOWS
-	FString ShutdownEventName = GetZenProcessShutdownEventName(EffectiveListenPort);
-	HANDLE Handle = OpenNativeEvent(ShutdownEventName, EVENT_MODIFY_STATE);
-	if (Handle == NULL)
+	ZenSharedEvent ShutDownEvent(ZenSharedEvent::GetShutdownEventName(EffectiveListenPort));
+	if (!ShutDownEvent.Open())
 	{
-		DWORD err = GetLastError();
-		UE_LOG(LogZenServiceInstance, Warning, TEXT("Failed opening named event '%s' (err: %d)"), *ShutdownEventName, err);
 		return false;
 	}
-	ON_SCOPE_EXIT{ CloseHandle(Handle); };
-	BOOL OK = SetEvent(Handle);
-	if (!OK)
+	if (!ShutDownEvent.Set())
 	{
-		DWORD err = GetLastError();
-		UE_LOG(LogZenServiceInstance, Warning, TEXT("Failed signaling named event '%s' (err: %d)"), *ShutdownEventName, err);
 		return false;
 	}
 	return true;
-	
-#elif PLATFORM_UNIX || PLATFORM_MAC
-	FAnsiString ShutdownEventName = GetZenProcessShutdownEventName(EffectiveListenPort);
-	int Semaphore = OpenNativeEvent(GetZenProcessShutdownEventName(EffectiveListenPort), 0600);
-	if (Semaphore < 0)
-	{
-		int err = errno;
-		UE_LOG(LogZenServiceInstance, Warning, TEXT("Failed opening named event '%hs' (err: %d)"), *ShutdownEventName, err);
-		return false;
-	}
-
-	if (semctl(Semaphore, 0, SETVAL, 0) < 0)
-	{
-		int err = errno;
-		UE_LOG(LogZenServiceInstance, Warning, TEXT("Failed signaling named event '%hs' (err: %d)"), *ShutdownEventName, err);
-		return false;
-	}
-	return true;
-#else
-	static_assert(false, "Missing implementation for Zen named shutdown events");
-	return false;
-#endif
-}
-
-static bool
-IsZenProcessActive(const TCHAR* ExecutablePath, uint32* OutPid)
-{
-	FString NormalizedExecutablePath(ExecutablePath);
-	FPaths::NormalizeFilename(NormalizedExecutablePath);
-	FPlatformProcess::FProcEnumerator ProcIter;
-	while (ProcIter.MoveNext())
-	{
-		FPlatformProcess::FProcEnumInfo ProcInfo = ProcIter.GetCurrent();
-		FString Candidate = ProcInfo.GetFullPath();
-		FPaths::NormalizeFilename(Candidate);
-		if (Candidate == NormalizedExecutablePath)
-		{
-			if (OutPid)
-			{
-				*OutPid = ProcInfo.GetPID();
-			}
-			return true;
-		}
-	}
-	return false;
-}
-
-static bool
-FindZenProcessId(uint32* OutPid)
-{
-	FString ExecutableName = GetLocalServiceExecutableName();
-	FPaths::NormalizeFilename(ExecutableName);
-	FPlatformProcess::FProcEnumerator ProcIter;
-	while (ProcIter.MoveNext())
-	{
-		FPlatformProcess::FProcEnumInfo ProcInfo = ProcIter.GetCurrent();
-		FString Candidate = ProcInfo.GetFullPath();
-		FPaths::NormalizeFilename(Candidate);
-		FString CandidateExecutableName = FPaths::GetPathLeaf(Candidate);
-		if (ExecutableName == CandidateExecutableName)
-		{
-			if (OutPid)
-			{
-				*OutPid = ProcInfo.GetPID();
-			}
-			return true;
-		}
-	}
-	return false;
 }
 
 static bool ShutdownZenServerProcess(int Pid, double MaximumWaitDurationSeconds = 25.0)
@@ -1423,7 +811,7 @@ static bool ShutdownZenServerProcess(int Pid, double MaximumWaitDurationSeconds
 		if (RequestZenShutdownOnEffectivePort(EffectivePort))
 		{
 			uint64 ZenShutdownWaitStartTime = FPlatformTime::Cycles64();
-			while (NativeIsProcessRunning(Pid))
+			while (ZenServerState::IsProcessRunning(Pid))
 			{
 				double ZenShutdownWaitDuration = FPlatformTime::ToSeconds64(FPlatformTime::Cycles64() - ZenShutdownWaitStartTime);
 				if (ZenShutdownWaitDuration < MaximumWaitDurationSeconds)
@@ -1438,12 +826,12 @@ static bool ShutdownZenServerProcess(int Pid, double MaximumWaitDurationSeconds
 			}
 		}
 	}
-	if (NativeIsProcessRunning(Pid))
+	if (ZenServerState::IsProcessRunning(Pid))
 	{
 		UE_LOG(LogZenServiceInstance, Display, TEXT("Attempting termination of zenserver process with pid %d"), Pid);
-		if (!NativeTerminate(Pid))
+		if (!ZenServerState::Terminate(Pid))
 		{
-			if (NativeIsProcessRunning(Pid))
+			if (ZenServerState::IsProcessRunning(Pid))
 			{
 				UE_LOG(LogZenServiceInstance, Warning, TEXT("Failed to terminate zenserver process with pid %d"), Pid);
 				return false;
@@ -1458,7 +846,7 @@ static bool ShutDownZenServerProcessExecutable(const FString& ExecutablePath, do
 {
 	uint64 ZenShutdownWaitStartTime = FPlatformTime::Cycles64();
 	uint32_t Pid = 0;
-	while (IsZenProcessActive(*ExecutablePath, &Pid))
+	while (ZenServerState::FindRunningProcessId(*ExecutablePath, &Pid))
 	{
 		UE_LOG(LogZenServiceInstance, Display, TEXT("Attempting to shut down of zenserver executable '%s' process with pid %d"), *ExecutablePath, Pid);
 		double ZenShutdownWaitDuration = FPlatformTime::ToSeconds64(FPlatformTime::Cycles64() - ZenShutdownWaitStartTime);
@@ -1480,24 +868,24 @@ static bool ShutDownZenServerProcessLockingDataDir(const FString& DataPath, doub
 	const FString LockFilePath = FPaths::Combine(DataPath, TEXT(".lock"));
 
 	uint64 ZenShutdownWaitStartTime = FPlatformTime::Cycles64();
-	if (!IsLockFileLocked(*LockFilePath, true))
+	if (!ZenLockFileData::IsLockFileLocked(*LockFilePath, true))
 	{
 		UE_LOG(LogZenServiceInstance, Display, TEXT("Lock file '%s' is not active, nothing to do"), *LockFilePath);
 		return true;
 	}
-	LockFileData LockFileState = ReadCbLockFile(LockFilePath);
+	ZenLockFileData LockFileState = ZenLockFileData::ReadCbLockFile(*LockFilePath);
 	if (!LockFileState.IsValid)
 	{
 		while (true)
 		{
-			if (!IsLockFileLocked(*LockFilePath, true))
+			if (!ZenLockFileData::IsLockFileLocked(*LockFilePath, true))
 			{
 				return true;
 			}
 			uint32_t Pid = 0;
-			if (!FindZenProcessId(&Pid))
+			if (!ZenServerState::FindRunningProcessId(*GetLocalServiceExecutableName(), & Pid))
 			{
-				if (!IsLockFileLocked(*LockFilePath, true))
+				if (!ZenLockFileData::IsLockFileLocked(*LockFilePath, true))
 				{
 					UE_LOG(LogZenServiceInstance, Display, TEXT("Lock file '%s' is no longer active, nothing to do"), *LockFilePath);
 					return true;
@@ -1512,7 +900,7 @@ static bool ShutDownZenServerProcessLockingDataDir(const FString& DataPath, doub
 				break;
 			}
 		}
-		if (!IsLockFileLocked(*LockFilePath))
+		if (!ZenLockFileData::IsLockFileLocked(*LockFilePath))
 		{
 			UE_LOG(LogZenServiceInstance, Display, TEXT("Successfully shut down zenserver using lock file '%s'"), *LockFilePath);
 			return true;
@@ -1530,7 +918,7 @@ static bool ShutDownZenServerProcessLockingDataDir(const FString& DataPath, doub
 		UE_LOG(LogZenServiceInstance, Display, TEXT("Requesting shut down of zenserver process using lock file '%s' with effective port %d"), *LockFilePath, EffectivePort);
 		if (RequestZenShutdownOnEffectivePort(EffectivePort))
 		{
-			while (IsLockFileLocked(*LockFilePath, true))
+			while (ZenLockFileData::IsLockFileLocked(*LockFilePath, true))
 			{
 				double ZenShutdownWaitDuration = FPlatformTime::ToSeconds64(FPlatformTime::Cycles64() - ZenShutdownWaitStartTime);
 				if (ZenShutdownWaitDuration < MaximumWaitDurationSeconds)
@@ -1543,7 +931,7 @@ static bool ShutDownZenServerProcessLockingDataDir(const FString& DataPath, doub
 					break;
 				}
 			}
-			if (!IsLockFileLocked(*LockFilePath, true))
+			if (!ZenLockFileData::IsLockFileLocked(*LockFilePath, true))
 			{
 				UE_LOG(LogZenServiceInstance, Display, TEXT("Successfully shut down zenserver process using lock file '%s' with effective port %u"), *LockFilePath, EffectivePort);
 				return true;
@@ -1553,14 +941,14 @@ static bool ShutDownZenServerProcessLockingDataDir(const FString& DataPath, doub
 
 	while (true)
 	{
-		if (!IsLockFileLocked(*LockFilePath, true))
+		if (!ZenLockFileData::IsLockFileLocked(*LockFilePath, true))
 		{
 			return true;
 		}
 		uint32_t Pid = 0;
-		if (!FindZenProcessId(&Pid))
+		if (!ZenServerState::FindRunningProcessId(*GetLocalServiceExecutableName(), &Pid))
 		{
-			if (!IsLockFileLocked(*LockFilePath, true))
+			if (!ZenLockFileData::IsLockFileLocked(*LockFilePath, true))
 			{
 				UE_LOG(LogZenServiceInstance, Display, TEXT("Lock file '%s' is no longer active, nothing to do"), *LockFilePath);
 				return true;
@@ -1576,7 +964,7 @@ static bool ShutDownZenServerProcessLockingDataDir(const FString& DataPath, doub
 		}
 	}
 
-	if (!IsLockFileLocked(*LockFilePath))
+	if (!ZenLockFileData::IsLockFileLocked(*LockFilePath))
 	{
 		UE_LOG(LogZenServiceInstance, Display, TEXT("Successfully shut down zenserver using lock file '%s'"), *LockFilePath);
 		return true;
@@ -1586,14 +974,14 @@ static bool ShutDownZenServerProcessLockingDataDir(const FString& DataPath, doub
 }
 
 static bool
-IsZenProcessUsingDataDir(const TCHAR* LockFilePath, LockFileData* OutLockFileData)
+IsZenProcessUsingDataDir(const TCHAR* LockFilePath, ZenLockFileData* OutLockFileData)
 {
-	if (IsLockFileLocked(LockFilePath, true))
+	if (ZenLockFileData::IsLockFileLocked(LockFilePath, true))
 	{
 		if (OutLockFileData)
 		{
 			// If an instance is running with this data path, check if we can use it and what port it is on
-			*OutLockFileData = ReadCbLockFile(LockFilePath);
+			*OutLockFileData = ZenLockFileData::ReadCbLockFile(LockFilePath);
 		}
 		return true;
 	}
@@ -1744,7 +1132,7 @@ bool
 IsLocalServiceRunning(const TCHAR* DataPath, uint16* OutEffectivePort)
 {
 	const FString LockFilePath = FPaths::Combine(DataPath, TEXT(".lock"));
-	LockFileData LockFileState;
+	ZenLockFileData LockFileState;
 	if (IsZenProcessUsingDataDir(*LockFilePath, &LockFileState))
 	{
 		if (OutEffectivePort != nullptr && LockFileState.IsValid && LockFileState.IsReady)
@@ -1843,8 +1231,8 @@ bool
 StopLocalService(const TCHAR* DataPath, double MaximumWaitDurationSeconds)
 {
 	const FString LockFilePath = FPaths::Combine(DataPath, TEXT(".lock"));
-	LockFileData LockFileState;
-	if (IsLockFileLocked(*LockFilePath, true))
+	ZenLockFileData LockFileState;
+	if (ZenLockFileData::IsLockFileLocked(*LockFilePath, true))
 	{
 		return ShutDownZenServerProcessLockingDataDir(DataPath, MaximumWaitDurationSeconds);
 	}
@@ -1868,18 +1256,6 @@ GetLocalServiceInstallVersion(bool bDetailed)
 	return InstallVersion.ToString(bDetailed);
 }
 
-FString
-GetLocalInstallUtilityPath()
-{
-	return FPaths::ConvertRelativePathToFull(FPaths::Combine(GetLocalZenRootPath(), TEXT("Zen\\Install"),
-#if PLATFORM_WINDOWS
-		TEXT("zen.exe")
-#else
-		TEXT("zen")
-#endif
-		));
-}
-
 static bool GIsDefaultServicePresent = false;
 
 FZenServiceInstance& GetDefaultServiceInstance()
@@ -2044,7 +1420,6 @@ FZenServiceInstance::TryRecovery()
 				UE_LOG(LogZenServiceInstance, Display, TEXT("Local ZenServer recovery being attempted..."));
 
 				bool bShutdownExistingInstance = true;
-				std::atomic<uint32> PreviousSponsorPids[UE_ARRAY_COUNT(ZenServerState::ZenServerEntry::SponsorPids)];
 				{
 					const ZenServerState ServerState(/* ReadOnly */true);
 					const ZenServerState::ZenServerEntry* Entry = ServerState.LookupByEffectiveListenPort(Port);
@@ -2055,11 +1430,6 @@ FZenServiceInstance::TryRecovery()
 							// The running process pid is not the same as the one we launched.  The process was relaunched elsewhere. Avoid shutting it down again.
 							bShutdownExistingInstance = false;
 						}
-
-						for (int32 SponsorPidIndex = 0; SponsorPidIndex < UE_ARRAY_COUNT(ZenServerState::ZenServerEntry::SponsorPids); ++SponsorPidIndex)
-						{
-							PreviousSponsorPids[SponsorPidIndex].store(Entry->SponsorPids[SponsorPidIndex].load(std::memory_order_relaxed), std::memory_order_relaxed);
-						}
 					}
 				}
 				if (bShutdownExistingInstance && !ShutdownZenServerProcess((int)AutoLaunchedPid))	// !ShutdownRunningServiceUsingEffectivePort(Port))
@@ -2096,15 +1466,6 @@ FZenServiceInstance::TryRecovery()
 				{
 					UE_LOG(LogZenServiceInstance, Display, TEXT("Local ZenServer post recovery status: NOT healthy"));
 				}
-				ZenServerState PostRecoveryServerState(/*bReadOnly*/ false);
-				ZenServerState::ZenServerEntry* PostRecoveryEntry = PostRecoveryServerState.LookupByEffectiveListenPort(Port);
-				if (PostRecoveryEntry)
-				{
-					for (int32 SponsorPidIndex = 0; SponsorPidIndex < UE_ARRAY_COUNT(ZenServerState::ZenServerEntry::SponsorPids); ++SponsorPidIndex)
-					{
-						PostRecoveryEntry->SponsorPids[SponsorPidIndex].store(PreviousSponsorPids[SponsorPidIndex].load(std::memory_order_relaxed), std::memory_order_relaxed);
-					}
-				}
 			}
 		}
 	}
@@ -2311,7 +1672,7 @@ FZenServiceInstance::ConditionalUpdateLocalInstall()
 #endif
 
 	FString InTreeCrashpadHandlerFilePath = FPaths::ConvertRelativePathToFull(FPlatformProcess::GenerateApplicationPath(TEXT("crashpad_handler"), EBuildConfiguration::Development));
-	FString InstallCrashpadHandlerFilePath = FPaths::ConvertRelativePathToFull(FPaths::Combine(GetLocalZenRootPath(), TEXT("Zen\\Install"), FString(FPathViews::GetCleanFilename(InTreeCrashpadHandlerFilePath))));
+	FString InstallCrashpadHandlerFilePath = FPaths::ConvertRelativePathToFull(FPaths::Combine(GetLocalInstallPath(), FString(FPathViews::GetCleanFilename(InTreeCrashpadHandlerFilePath))));
 
 	if (FileManager.FileExists(*InTreeCrashpadHandlerFilePath) && (bMainExecutablesUpdated || !FileManager.FileExists(*InstallCrashpadHandlerFilePath)))
 	{
@@ -2330,7 +1691,7 @@ FZenServiceInstance::AutoLaunch(const FServiceAutoLaunchSettings& InSettings, FS
 
 	FString WorkingDirectory = FPaths::GetPath(ExecutablePath);
 
-	LockFileData LockFileState;
+	ZenLockFileData LockFileState;
 	uint64 ZenWaitForRunningProcessReadyStartTime = FPlatformTime::Cycles64();
 	while (IsZenProcessUsingDataDir(*LockFilePath, &LockFileState) && LockFileState.IsValid && !LockFileState.IsReady)
 	{
@@ -2341,7 +1702,7 @@ FZenServiceInstance::AutoLaunch(const FServiceAutoLaunchSettings& InSettings, FS
 			break;
 		}
 		FPlatformProcess::Sleep(0.1f);
-		LockFileState = LockFileData();
+		LockFileState = ZenLockFileData();
 	}
 
 	bool bShutDownExistingInstanceForDataPath = true;
@@ -2443,8 +1804,6 @@ FZenServiceInstance::AutoLaunch(const FServiceAutoLaunchSettings& InSettings, FS
 		}
 	}
 
-	// When limiting process lifetime, always re-launch to add sponsor process IDs.
-	// When not limiting process lifetime, only launch if the process is not already live.
 	if (bLaunchNewInstance)
 	{
 		if (InSettings.bIsDefaultDataPath && InSettings.bIsDefaultSharedRunContext)
@@ -2486,15 +1845,103 @@ FZenServiceInstance::AutoLaunch(const FServiceAutoLaunchSettings& InSettings, FS
 		EffectiveRunContext.DataPath = InSettings.DataPath;
 		EffectiveRunContext.bShowConsole = InSettings.bShowConsole;
 
+		FString StartupEventName = ZenSharedEvent::GetStartupEventName();
+		ZenSharedEvent StartupEvent(StartupEventName);
+		if (!StartupEvent.Create())
+		{
+			return false;
+		}
+		TransientParms.Appendf(TEXT(" --child-id %s"), *StartupEventName);
+
 		FProcHandle Proc = StartLocalService(EffectiveRunContext, TransientParms.IsEmpty() ? nullptr : *TransientParms);
 
-		// Only write run context if we're using the default shared run context
-		if (InSettings.bIsDefaultSharedRunContext)
+		if (Proc.IsValid())
 		{
-			EffectiveRunContext.WriteToJsonFile(*ExecutionContextFilePath);
-		}
+			// Only write run context if we're using the default shared run context
+			if (InSettings.bIsDefaultSharedRunContext)
+			{
+				EffectiveRunContext.WriteToJsonFile(*ExecutionContextFilePath);
+			}
 
-		if (!Proc.IsValid())
+			bool ZenServerIsReady = false;
+			FScopedSlowTask WaitForZenReadySlowTask(0, NSLOCTEXT("Zen", "Zen_WaitingForReady", "Waiting for ZenServer to be ready"));
+			uint64 ZenWaitStartTime = FPlatformTime::Cycles64();
+
+			enum class EWaitDurationPhase
+			{
+				Short,
+				Medium,
+				Long
+			} DurationPhase = EWaitDurationPhase::Short;
+
+			while (FPlatformProcess::IsProcRunning(Proc))
+			{
+				if (StartupEvent.Wait(5000))
+				{
+					ZenServerIsReady = true;
+					break;
+				}
+
+				double ZenWaitDuration = FPlatformTime::ToSeconds64(FPlatformTime::Cycles64() - ZenWaitStartTime);
+				if (ZenWaitDuration >= 5.0)
+				{
+					if (DurationPhase == EWaitDurationPhase::Short)
+					{
+#if !IS_PROGRAM
+						if (!FApp::IsUnattended() && !IsRunningCommandlet() && !GIsRunningUnattendedScript)
+						{
+							FText ZenLaunchFailurePromptTitle = NSLOCTEXT("Zen", "Zen_LaunchFailurePromptTitle", "Failed to launch");
+
+							FFormatNamedArguments FormatArguments;
+							FString LogFilePath = FPaths::Combine(InSettings.DataPath, TEXT("logs"), TEXT("zenserver.log"));
+							FPaths::MakePlatformFilename(LogFilePath);
+							FormatArguments.Add(TEXT("LogFilePath"), FText::FromString(LogFilePath));
+							FText ZenLaunchFailurePromptText = FText::Format(NSLOCTEXT("Zen", "Zen_LaunchFailurePromptText", "Unreal Zen Storage Server failed to launch. Please check the ZenServer log file for details:\n{LogFilePath}"), FormatArguments);
+							FPlatformMisc::MessageBoxExt(EAppMsgType::Ok, *ZenLaunchFailurePromptText.ToString(), *ZenLaunchFailurePromptTitle.ToString());
+							break;
+						}
+						else
+#endif
+						{
+							// Just log as there is no one to show a message
+							UE_LOG(LogZenServiceInstance, Warning, TEXT("Unreal Zen Storage Server did not launch in the expected duration"));
+							break;
+						}
+						// Note that the dialog may not show up when zenserver is needed early in the launch cycle, but this will at least ensure
+						// the splash screen is refreshed with the appropriate text status message.
+						WaitForZenReadySlowTask.MakeDialog(true, false);
+						UE_LOG(LogZenServiceInstance, Display, TEXT("Waiting for ZenServer to be ready..."));
+						DurationPhase = EWaitDurationPhase::Medium;
+					}
+#if !IS_PROGRAM
+					else if (!(FApp::IsUnattended() || IsRunningCommandlet() || GIsRunningUnattendedScript) && ZenWaitDuration > 20.0 && (DurationPhase == EWaitDurationPhase::Medium))
+					{
+						FText ZenLongWaitPromptTitle = NSLOCTEXT("Zen", "Zen_LongWaitPromptTitle", "Wait for ZenServer?");
+						FText ZenLongWaitPromptText = NSLOCTEXT("Zen", "Zen_LongWaitPromptText", "Unreal Zen Storage Server is taking a long time to launch. It may be performing maintenance. Keep waiting?");
+						if (FPlatformMisc::MessageBoxExt(EAppMsgType::YesNo, *ZenLongWaitPromptText.ToString(), *ZenLongWaitPromptTitle.ToString()) == EAppReturnType::No)
+						{
+							break;
+						}
+						DurationPhase = EWaitDurationPhase::Long;
+					}
+#endif
+					if (WaitForZenReadySlowTask.ShouldCancel())
+					{
+						break;
+					}
+				}
+			}
+
+			if (!ZenServerIsReady)
+			{
+				if (FPlatformProcess::IsProcRunning(Proc))
+				{
+					UE_LOG(LogZenServiceInstance, Warning, TEXT("Terminating unresponsive process for executable '%s' on port %u"), *ExecutablePath, InSettings.DesiredPort);
+					FPlatformProcess::TerminateProc(Proc, true);
+				}
+			}
+		}
+		else
 		{
 			UE_LOG(LogZenServiceInstance, Warning, TEXT("Failed launch service using executable '%s' on port %u"), *ExecutablePath, InSettings.DesiredPort);
 			return false;
@@ -2523,88 +1970,17 @@ FZenServiceInstance::AutoLaunch(const FServiceAutoLaunchSettings& InSettings, FS
 	// Default to assuming that we get to run on the port we want
 	OutPort = InSettings.DesiredPort;
 
-	FScopedSlowTask WaitForZenReadySlowTask(0, NSLOCTEXT("Zen", "Zen_WaitingForReady", "Waiting for ZenServer to be ready"));
-	uint64 ZenWaitStartTime = FPlatformTime::Cycles64();
-	enum class EWaitDurationPhase
+	ZenLockFileData RunningLockFileState = ZenLockFileData::ReadCbLockFile(*LockFilePath);
+	if (!RunningLockFileState.IsValid)
 	{
-		Short,
-		Medium,
-		Long
-	} DurationPhase = EWaitDurationPhase::Short;
-	bool bIsReady = false;
-	while (!bIsReady)
+		return false;
+	}
+	if (!RunningLockFileState.IsReady)
 	{
-		LockFileData RunningLockFileState = ReadCbLockFile(LockFilePath);
-		if (RunningLockFileState.IsValid)
-		{
-			bIsReady = RunningLockFileState.IsReady;
-			if (bIsReady)
-			{
-				OutPort = RunningLockFileState.EffectivePort;
-				break;
-			}
-		}
-
-		double ZenWaitDuration = FPlatformTime::ToSeconds64(FPlatformTime::Cycles64() - ZenWaitStartTime);
-		if (ZenWaitDuration < 10.0)
-		{
-			// Initial 10 second window of higher frequency checks
-			FPlatformProcess::Sleep(0.01f);
-		}
-		else
-		{
-			if (DurationPhase == EWaitDurationPhase::Short)
-			{
-				if (!IsLockFileLocked(*LockFilePath))
-				{
-#if !IS_PROGRAM
-					if (!FApp::IsUnattended() && !IsRunningCommandlet() && !GIsRunningUnattendedScript)
-					{
-						FText ZenLaunchFailurePromptTitle = NSLOCTEXT("Zen", "Zen_LaunchFailurePromptTitle", "Failed to launch");
-
-						FFormatNamedArguments FormatArguments;
-						FString LogFilePath = FPaths::Combine(InSettings.DataPath, TEXT("logs"), TEXT("zenserver.log"));
-						FPaths::MakePlatformFilename(LogFilePath);
-						FormatArguments.Add(TEXT("LogFilePath"), FText::FromString(LogFilePath));
-						FText ZenLaunchFailurePromptText = FText::Format(NSLOCTEXT("Zen", "Zen_LaunchFailurePromptText", "Unreal Zen Storage Server failed to launch. Please check the ZenServer log file for details:\n{LogFilePath}"), FormatArguments);
-						FPlatformMisc::MessageBoxExt(EAppMsgType::Ok, *ZenLaunchFailurePromptText.ToString(), *ZenLaunchFailurePromptTitle.ToString());
-						return false;
-					}
-					else
-#endif
-					{
-						// Just log as there is no one to show a message
-						UE_LOG(LogZenServiceInstance, Warning, TEXT("Unreal Zen Storage Server did not launch in the expected duration"));
-						return false;
-					}
-				}
-				// Note that the dialog may not show up when zenserver is needed early in the launch cycle, but this will at least ensure
-				// the splash screen is refreshed with the appropriate text status message.
-				WaitForZenReadySlowTask.MakeDialog(true, false);
-				UE_LOG(LogZenServiceInstance, Display, TEXT("Waiting for ZenServer to be ready..."));
-				DurationPhase = EWaitDurationPhase::Medium;
-			}
-#if !IS_PROGRAM
-			else if (!(FApp::IsUnattended() || IsRunningCommandlet() || GIsRunningUnattendedScript) && ZenWaitDuration > 20.0 && (DurationPhase == EWaitDurationPhase::Medium))
-			{
-				FText ZenLongWaitPromptTitle = NSLOCTEXT("Zen", "Zen_LongWaitPromptTitle", "Wait for ZenServer?");
-				FText ZenLongWaitPromptText = NSLOCTEXT("Zen", "Zen_LongWaitPromptText", "Unreal Zen Storage Server is taking a long time to launch. It may be performing maintenance. Keep waiting?");
-				if (FPlatformMisc::MessageBoxExt(EAppMsgType::YesNo, *ZenLongWaitPromptText.ToString(), *ZenLongWaitPromptTitle.ToString()) == EAppReturnType::No)
-				{
-					return false;
-				}
-				DurationPhase = EWaitDurationPhase::Long;
-			}
-#endif
-
-			if (WaitForZenReadySlowTask.ShouldCancel())
-			{
-				return false;
-			}
-			FPlatformProcess::Sleep(0.1f);
-		}
+		return false;
 	}
-	return bIsReady;
+	OutPort = RunningLockFileState.EffectivePort;
+	return true;
 }
 
 bool 
diff --git a/Engine/Source/Developer/Zen/Private/ZenServerState.cpp b/Engine/Source/Developer/Zen/Private/ZenServerState.cpp
new file mode 100644
--- /dev/null
+++ b/Engine/Source/Developer/Zen/Private/ZenServerState.cpp
@@ -0,0 +1,817 @@
+// Copyright Epic Games, Inc. All Rights Reserved.
+
+#include "ZenServerState.h"
+
+#include "HAL/FileManager.h"
+#include "HAL/PlatformProcess.h"
+#include "Logging/LogMacros.h"
+#include "Misc/ScopeExit.h"
+#include "Misc/Paths.h"
+#include "Misc/PathViews.h"
+#include "Serialization/CompactBinaryValidation.h"
+
+#if PLATFORM_WINDOWS
+#	include "Windows/AllowWindowsPlatformTypes.h"
+#	include <shellapi.h>
+#	include <synchapi.h>
+#	include <Windows.h>
+#	include "Windows/HideWindowsPlatformTypes.h"
+#endif
+
+#if PLATFORM_UNIX || PLATFORM_MAC
+#	include <sys/file.h>
+#	include <sys/mman.h>
+#	include <sys/sem.h>
+#endif
+
+DEFINE_LOG_CATEGORY_STATIC(LogZenServiceState, Log, All);
+
+// Native functions to interact with a process using a process id
+// We don't use UE's own OpenProcess as they try to open processes with PROCESS_ALL_ACCESS
+bool ZenServerState::IsProcessRunning(uint32 Pid)
+{
+	if (Pid == 0)
+	{
+		return false;
+	}
+#if PLATFORM_WINDOWS
+	HANDLE Handle = OpenProcess(PROCESS_QUERY_LIMITED_INFORMATION, 0, (DWORD)Pid);
+
+	if (!Handle)
+	{
+		DWORD Error = GetLastError();
+
+		if (Error == ERROR_INVALID_PARAMETER)
+		{
+			return false;
+		}
+		else if (Error == ERROR_ACCESS_DENIED)
+		{
+			UE_LOG(LogZenServiceState, Warning, TEXT("No access to open running process %d: %d, assuming it is running"), Pid, Error);
+			return true;
+		}
+		UE_LOG(LogZenServiceState, Warning, TEXT("Failed to open running process %d: %d, assuming it is not running"), Pid, Error);
+		return false;
+	}
+	ON_SCOPE_EXIT{ CloseHandle(Handle); };
+
+	DWORD ExitCode = 0;
+	if (GetExitCodeProcess(Handle, &ExitCode) == 0)
+	{
+		DWORD Error = GetLastError();
+		UE_LOG(LogZenServiceState, Warning, TEXT("Failed to get running process exit code %d: %d, assuming it is still running"), Pid, Error);
+		return true;
+	}
+	else if (ExitCode == STILL_ACTIVE)
+	{
+		return true;
+	}
+
+	return false;
+
+#elif PLATFORM_UNIX || PLATFORM_MAC
+	int Res = kill(pid_t(Pid), 0);
+	if (Res == 0)
+	{
+		return true;
+	}
+	int Error = errno;
+	if (Error == EPERM)
+	{
+		UE_LOG(LogZenServiceState, Warning, TEXT("No permission to signal running process %d: %d, assuming it is running"), Pid, Error);
+		return true;
+	}
+	else if (Error == ESRCH)
+	{
+		return false;
+	}
+	UE_LOG(LogZenServiceState, Warning, TEXT("Failed to signal running process %d: %d, assuming it is running"), Pid, Error);
+	return true;
+#endif
+}
+
+bool ZenServerState::Terminate(uint32 Pid)
+{
+#if PLATFORM_WINDOWS
+	HANDLE Handle = OpenProcess(PROCESS_TERMINATE | SYNCHRONIZE, 0, (DWORD)Pid);
+	if (Handle == NULL)
+	{
+		DWORD Error = GetLastError();
+
+		if (Error != ERROR_INVALID_PARAMETER)
+		{
+			UE_LOG(LogZenServiceState, Warning, TEXT("Failed to open running process for terminate %d: %d"), Pid, Error);
+			return false;
+		}
+		return true;
+	}
+	ON_SCOPE_EXIT{ CloseHandle(Handle); };
+
+	BOOL bTerminated = TerminateProcess(Handle, 0);
+	if (!bTerminated)
+	{
+		DWORD Error = GetLastError();
+		UE_LOG(LogZenServiceState, Warning, TEXT("Failed to terminate running process %d: %d"), Pid, Error);
+		return false;
+	}
+	DWORD WaitResult = WaitForSingleObject(Handle, 15000);
+	BOOL bSuccess = (WaitResult == WAIT_OBJECT_0) || (WaitResult == WAIT_ABANDONED_0);
+	if (!bSuccess)
+	{
+		if (WaitResult == WAIT_FAILED)
+		{
+			DWORD Error = GetLastError();
+			UE_LOG(LogZenServiceState, Warning, TEXT("Failed to wait for terminated process %d: %d"), Pid, Error);
+		}
+		return false;
+	}
+#elif PLATFORM_UNIX || PLATFORM_MAC
+	int Res = kill(pid_t(Pid), SIGKILL);
+	if (Res != 0)
+	{
+		int err = errno;
+		if (err != ESRCH)
+		{
+			UE_LOG(LogZenServiceState, Warning, TEXT("Failed to terminate running process %d: %d"), Pid, err);
+			return false;
+		}
+	}
+#endif
+	return true;
+}
+
+bool ZenServerState::FindRunningProcessId(const TCHAR* ExecutablePath, uint32* OutPid)
+{
+	FString NormalizedExecutablePath(ExecutablePath);
+	FPaths::NormalizeFilename(NormalizedExecutablePath);
+	FPlatformProcess::FProcEnumerator ProcIter;
+	while (ProcIter.MoveNext())
+	{
+		FPlatformProcess::FProcEnumInfo ProcInfo = ProcIter.GetCurrent();
+		FString Candidate = ProcInfo.GetFullPath();
+		FPaths::NormalizeFilename(Candidate);
+		if (Candidate == NormalizedExecutablePath)
+		{
+			if (OutPid)
+			{
+				*OutPid = ProcInfo.GetPID();
+			}
+			return true;
+		}
+	}
+	return false;
+}
+
+ZenServerState::ZenServerState(bool ReadOnly)
+	: m_hMapFile(nullptr)
+	, m_Data(nullptr)
+{
+	size_t MapSize = m_MaxEntryCount * sizeof(ZenServerEntry);
+
+#if PLATFORM_WINDOWS
+	DWORD DesiredAccess = ReadOnly ? FILE_MAP_READ : (FILE_MAP_READ | FILE_MAP_WRITE);
+	HANDLE hMap = OpenFileMapping(DesiredAccess, 0, L"Global\\ZenMap");
+	if (hMap == NULL)
+	{
+		hMap = OpenFileMapping(DesiredAccess, 0, L"Local\\ZenMap");
+	}
+
+	if (hMap == NULL)
+	{
+		return;
+	}
+
+	void* pBuf = MapViewOfFile(hMap,		   // handle to map object
+		DesiredAccess,  // read permission
+		0,			   // offset high
+		0,			   // offset low
+		MapSize);
+
+	if (pBuf == NULL)
+	{
+		CloseHandle(hMap);
+		return;
+	}
+#elif PLATFORM_UNIX || PLATFORM_MAC
+	int OFlag = ReadOnly ? (O_RDONLY | O_CLOEXEC) : (O_RDWR | O_CREAT | O_CLOEXEC);
+	int Fd = shm_open("/UnrealEngineZen", OFlag, 0666);
+	if (Fd < 0)
+	{
+		return;
+	}
+	void* hMap = (void*)intptr_t(Fd);
+
+	int Prot = ReadOnly ? PROT_READ : (PROT_WRITE | PROT_READ);
+	void* pBuf = mmap(nullptr, MapSize, Prot, MAP_SHARED, Fd, 0);
+	if (pBuf == MAP_FAILED)
+	{
+		close(Fd);
+		return;
+	}
+#endif
+
+#if PLATFORM_WINDOWS || PLATFORM_UNIX || PLATFORM_MAC
+	m_hMapFile = hMap;
+	m_Data = reinterpret_cast<ZenServerEntry*>(pBuf);
+#endif
+	m_IsReadOnly = ReadOnly;
+}
+
+ZenServerState::~ZenServerState()
+{
+#if PLATFORM_WINDOWS
+	if (m_Data)
+	{
+		UnmapViewOfFile(m_Data);
+	}
+
+	if (m_hMapFile)
+	{
+		CloseHandle(m_hMapFile);
+	}
+#elif PLATFORM_UNIX || PLATFORM_MAC
+	if (m_Data != nullptr)
+	{
+		munmap((void*)m_Data, m_MaxEntryCount * sizeof(ZenServerEntry));
+	}
+
+	int Fd = int(intptr_t(m_hMapFile));
+	close(Fd);
+#endif
+	m_hMapFile = nullptr;
+	m_Data = nullptr;
+}
+
+const ZenServerState::ZenServerEntry* ZenServerState::LookupByDesiredListenPortInternal(int Port) const
+{
+	if (m_Data == nullptr)
+	{
+		return nullptr;
+	}
+
+	for (int i = 0; i < m_MaxEntryCount; ++i)
+	{
+		if (m_Data[i].DesiredListenPort.load(std::memory_order_relaxed) == Port)
+		{
+			const ZenServerState::ZenServerEntry* Entry = &m_Data[i];
+			if (IsProcessRunning((uint32)Entry->Pid.load(std::memory_order_relaxed)))
+			{
+				return Entry;
+			}
+		}
+	}
+
+	return nullptr;
+}
+
+const ZenServerState::ZenServerEntry* ZenServerState::LookupByDesiredListenPort(int Port) const
+{
+	return LookupByDesiredListenPortInternal(Port);
+}
+
+ZenServerState::ZenServerEntry* ZenServerState::LookupByDesiredListenPort(int Port)
+{
+	check(!m_IsReadOnly);
+	return const_cast<ZenServerState::ZenServerEntry*>(LookupByDesiredListenPortInternal(Port));
+}
+
+const ZenServerState::ZenServerEntry* ZenServerState::LookupByEffectiveListenPortInternal(int Port) const
+{
+	if (m_Data == nullptr)
+	{
+		return nullptr;
+	}
+
+	for (int i = 0; i < m_MaxEntryCount; ++i)
+	{
+		const ZenServerState::ZenServerEntry* Entry = &m_Data[i];
+		if (Entry->EffectiveListenPort.load(std::memory_order_relaxed) == Port)
+		{
+			if (IsProcessRunning((uint32)Entry->Pid.load(std::memory_order_relaxed)))
+			{
+				return Entry;
+			}
+		}
+	}
+
+	return nullptr;
+}
+
+const ZenServerState::ZenServerEntry* ZenServerState::LookupByEffectiveListenPort(int Port) const
+{
+	return LookupByEffectiveListenPortInternal(Port);
+}
+
+ZenServerState::ZenServerEntry* ZenServerState::LookupByEffectiveListenPort(int Port)
+{
+	check(!m_IsReadOnly);
+	return const_cast<ZenServerState::ZenServerEntry*>(LookupByEffectiveListenPortInternal(Port));
+}
+
+const ZenServerState::ZenServerEntry* ZenServerState::LookupByPid(uint32 Pid) const
+{
+	if (m_Data == nullptr)
+	{
+		return nullptr;
+	}
+
+	for (int i = 0; i < m_MaxEntryCount; ++i)
+	{
+		const ZenServerState::ZenServerEntry* Entry = &m_Data[i];
+		if (m_Data[i].Pid.load(std::memory_order_relaxed) == Pid)
+		{
+			if (IsProcessRunning(Pid))
+			{
+				return Entry;
+			}
+		}
+	}
+
+	return nullptr;
+}
+
+bool
+ZenServerState::ZenServerEntry::AddSponsorProcess(uint32 PidToAdd)
+{
+	for (std::atomic<uint32>& PidEntry : SponsorPids)
+	{
+		if (PidEntry.load(std::memory_order_relaxed) == 0)
+		{
+			uint32 Expected = 0;
+			if (PidEntry.compare_exchange_strong(Expected, PidToAdd))
+			{
+				// Success!
+				return true;
+			}
+		}
+		else if (PidEntry.load(std::memory_order_relaxed) == PidToAdd)
+		{
+			// Success, the because pid is already in the list
+			return true;
+		}
+	}
+
+	return false;
+}
+
+ZenSharedEvent::ZenSharedEvent(FStringView EventName)
+	: m_EventName(EventName)
+{
+	check(m_EventName.Len() > 0);
+}
+ZenSharedEvent::~ZenSharedEvent()
+{
+	Close();
+}
+
+bool ZenSharedEvent::Create()
+{
+#if PLATFORM_WINDOWS
+	check(m_EventHandle == NULL);
+	FString FullEventName = GetFullEventName();
+	m_EventHandle = (void*)CreateEventW(nullptr, true, false, *FullEventName);
+	if (m_EventHandle == NULL)
+	{
+		DWORD LastError = GetLastError();
+		UE_LOG(LogZenServiceState, Warning, TEXT("Failed creating named event '{EventName}' (err: {Error})"), *FullEventName, LastError);
+		return false;
+	}
+	return true;
+#elif PLATFORM_UNIX || PLATFORM_MAC
+	check(m_Fd == -1);
+	check(m_Semaphore == -1);
+	FAnsiString EventPath = GetEventPath();
+	// Create a file to back the semaphore
+	m_Fd = open(*EventPath, O_RDWR | O_CREAT | O_CLOEXEC, 0666);
+	if (m_Fd < 0)
+	{
+		int LastError = errno;
+		UE_LOG(LogZenServiceState, Warning, TEXT("Failed to create named event '{EventName}' (err: {Error})"), *EventPath, LastError);
+		return false;
+	}
+	fchmod(m_Fd, 0666);
+
+	// Use the file path to generate an IPC key
+	key_t IpcKey = ftok(*EventPath, 1);
+	if (IpcKey < 0)
+	{
+		int LastError = errno;
+		close(m_Fd);
+		m_Fd = -1;
+		UE_LOG(LogZenServiceState, Warning, TEXT("Failed to create an SysV IPC key for named event '{EventName}' (err: {Error})"), *EventPath, LastError);
+		return false;
+	}
+
+	// Use the key to create/open the semaphore
+	m_Semaphore = semget(IpcKey, 1, 0600 | IPC_CREAT);
+	if (m_Semaphore < 0)
+	{
+		int LastError = errno;
+		close(m_Fd);
+		m_Fd = -1;
+		UE_LOG(LogZenServiceState, Warning, TEXT("Failed creating an SysV semaphore for named event '{EventName}' (err: {Error})"), *EventPath, LastError);
+		return false;
+	}
+
+	// Atomically claim ownership of the semaphore's key. The owner initializes
+	// the semaphore to 1 so we can use the wait-for-zero op as that does not
+	// modify the semaphore's value on a successful wait.
+	int LockResult = flock(m_Fd, LOCK_EX | LOCK_NB);
+	if (LockResult == 0)
+	{
+		// This isn't thread safe really. Another thread could open the same
+		// semaphore and successfully wait on it in the period of time where
+		// this comment is but before the semaphore's initialised.
+		semctl(m_Semaphore, 0, SETVAL, 1);
+	}
+	return true;
+#endif
+}
+
+bool ZenSharedEvent::Exists() const
+{
+#if PLATFORM_WINDOWS
+	FString FullEventName = GetFullEventName();
+	HANDLE EventHandle = OpenEventW(READ_CONTROL, false, *FullEventName);
+	if (EventHandle == NULL)
+	{
+		DWORD LastError = GetLastError();
+		if (LastError != ERROR_FILE_NOT_FOUND)
+		{
+			UE_LOG(LogZenServiceState, Warning, TEXT("Failed checking existance of named event '{EventName}' (err: {Error})"), *FullEventName, LastError);
+		}
+		return false;
+	}
+	else
+	{
+		CloseHandle(EventHandle);
+		return false;
+	}
+#elif PLATFORM_UNIX || PLATFORM_MAC
+	FAnsiString EventPath = GetEventPath();
+	key_t IpcKey = ftok(*EventPath, 1);
+	if (IpcKey < 0)
+	{
+		int LastError = errno;
+		UE_LOG(LogZenServiceState, Warning, TEXT("Failed to create an SysV IPC key for named event '{EventName}' (err: {Error})"), *EventPath, LastError);
+		return false;
+	}
+	int Semaphore = semget(IpcKey, 1, 0400);
+	if (Semaphore < 0)
+	{
+		int LastError = errno;
+		if (LastError != ENOENT)
+		{
+			UE_LOG(LogZenServiceState, Warning, TEXT("Failed checking named event '{EventName}' (err: {Error})"), *EventPath, LastError);
+		}
+		return false;
+	}
+	return true;
+#endif
+}
+
+bool ZenSharedEvent::Open()
+{
+#if PLATFORM_WINDOWS
+	check(m_EventHandle == NULL);
+	FString FullEventName = GetFullEventName();
+	m_EventHandle = (void*)OpenEventW(EVENT_MODIFY_STATE, false, *FullEventName);
+	if (m_EventHandle == NULL)
+	{
+		DWORD LastError = GetLastError();
+		if (LastError != ERROR_FILE_NOT_FOUND)
+		{
+			UE_LOG(LogZenServiceState, Warning, TEXT("Failed opening named event '{EventName}' (err: {Error})"), *FullEventName, LastError);
+		}
+		return false;
+	}
+	return true;
+#elif PLATFORM_UNIX || PLATFORM_MAC
+	check(m_Fd == -1);
+	check(m_Semaphore == -1);
+	FAnsiString EventPath = GetEventPath();
+	key_t IpcKey = ftok(*EventPath, 1);
+	if (IpcKey < 0)
+	{
+		int LastError = errno;
+		UE_LOG(LogZenServiceState, Warning, TEXT("Failed to create an SysV IPC key for named event '{EventName}' (err: {Error})"), *EventPath, LastError);
+		return false;
+	}
+	m_Semaphore = semget(IpcKey, 1, 0600);
+	if (m_Semaphore < 0)
+	{
+		int LastError = errno;
+		if (LastError != ENOENT)
+		{
+			UE_LOG(LogZenServiceState, Warning, TEXT("Failed opening named event '{EventName}' (err: {Error})"), *EventPath, LastError);
+		}
+		return false;
+	}
+	return true;
+#endif
+}
+
+bool ZenSharedEvent::Wait(int TimeoutMs)
+{
+#if PLATFORM_WINDOWS
+	check(m_EventHandle != NULL);
+
+	const DWORD Timeout = (TimeoutMs < 0) ? INFINITE : TimeoutMs;
+
+	DWORD Result = WaitForSingleObject((HANDLE)m_EventHandle, Timeout);
+
+	if (Result == WAIT_FAILED)
+	{
+		DWORD LastError = GetLastError();
+		UE_LOG(LogZenServiceState, Warning, TEXT("Failed waiting for named event '{EventName}' (err: {Error})"), *m_EventName, LastError);
+		return false;
+	}
+
+	return (Result == WAIT_OBJECT_0) || (Result == WAIT_ABANDONED_0);
+
+#elif PLATFORM_UNIX || PLATFORM_MAC
+	check(m_Semaphore != -1);
+
+	int			  Result;
+	struct sembuf SemOp = {};
+
+	if (TimeoutMs < 0)
+	{
+		Result = semop(m_Semaphore, &SemOp, 1);
+		if (Result != 0)
+		{
+			int LastError = errno;
+			UE_LOG(LogZenServiceState, Warning, TEXT("Failed waiting for named event '{EventName}' (err: {Error})"), *m_EventName, LastError);
+			return false;
+		}
+		return true;
+	}
+#if PLATFORM_UNIX
+	const int		TimeoutSec = TimeoutMs / 1000;
+	struct timespec TimeoutValue = {
+		.tv_sec = TimeoutSec,
+		.tv_nsec = (TimeoutMs - (TimeoutSec * 1000)) * 1000000,
+	};
+	Result = semtimedop(m_Semaphore, &SemOp, 1, &TimeoutValue);
+	if (Result == 0)
+	{
+		return true;
+	}
+	int LastError = errno;
+	if (LastError != EAGAIN)
+	{
+		UE_LOG(LogZenServiceState, Warning, TEXT("Failed waiting for named event '{EventName}' (err: {Error})"), *m_EventName, LastError);
+	}
+	return false;
+#elif PLATFORM_MAC
+	const unsigned int SleepTimeMs = 10u;
+	SemOp.sem_flg = IPC_NOWAIT;
+	do
+	{
+		Result = semop(m_Semaphore, &SemOp, 1);
+		if (Result == 0)
+		{
+			return true;
+		}
+		else
+		{
+			int LastError = errno;
+			if (errno != EAGAIN)
+			{
+				UE_LOG(LogZenServiceState, Warning, TEXT("Failed waiting for named event '{EventName}' (err: {Error})"), *m_EventName, LastError);
+				break;
+			}
+		}
+		usleep(SleepTimeMs * 1000u);
+		TimeoutMs -= SleepTimeMs;
+	} while (TimeoutMs > 0);
+	return false;
+#endif // PLATFORM_MAC
+#endif
+}
+
+bool ZenSharedEvent::Set()
+{
+#if PLATFORM_WINDOWS
+	check(m_EventHandle != nullptr);
+	if (SetEvent((HANDLE)m_EventHandle))
+	{
+		return true;
+	}
+	else
+	{
+		DWORD LastError = GetLastError();
+		UE_LOG(LogZenServiceState, Warning, TEXT("Failed signalling named event '{EventName}' (err: {Error})"), *m_EventName, LastError);
+		return false;
+	}
+#elif PLATFORM_UNIX || PLATFORM_MAC
+	check(m_Semaphore != -1);
+	if (semctl(m_Semaphore, 0, SETVAL, 0) != -1)
+	{
+		return true;
+	}
+	else
+	{
+		int LastError = errno;
+		UE_LOG(LogZenServiceState, Warning, TEXT("Failed signalling named event '{EventName}' (err: {Error})"), *m_EventName, LastError);
+		return false;
+	}
+#endif
+}
+
+void ZenSharedEvent::Close()
+{
+#if PLATFORM_WINDOWS
+	if (m_EventHandle != NULL)
+	{
+		CloseHandle((HANDLE)m_EventHandle);
+	}
+	m_EventHandle = NULL;
+#elif PLATFORM_UNIX || PLATFORM_MAC
+	if (m_Fd != -1)
+	{
+		if (flock(m_Fd, LOCK_EX | LOCK_NB) == 0)
+		{
+			unlink(*m_EventPath);
+
+			flock(m_Fd, LOCK_UN | LOCK_NB);
+			close(m_Fd);
+
+			semctl(m_Semaphore, 0, IPC_RMID);
+		}
+	}
+	m_Fd = -1;
+	m_Semaphore = -1;
+#endif
+}
+
+#if PLATFORM_WINDOWS
+
+FString ZenSharedEvent::GetFullEventName() const
+{
+	TStringBuilder<128> Name;
+	Name << "Local\\";
+	Name << m_EventName;
+	return *Name;
+}
+
+#elif PLATFORM_UNIX || PLATFORM_MAC
+
+FAnsiString ZenSharedEvent::GetEventPath() const
+{
+	TAnsiStringBuilder<128> Name;
+	Name << "/tmp/";
+	Name << m_EventName;
+	return *Name;
+}
+#endif
+
+FString ZenSharedEvent::GetShutdownEventName(uint16 EffectiveListenPort)
+{
+	return *WriteToWideString<64>(WIDETEXT("Zen_"), EffectiveListenPort, WIDETEXT("_Shutdown"));
+}
+
+FString ZenSharedEvent::GetStartupEventName()
+{
+	return (*WriteToWideString<64>(WIDETEXT("Zen_"), FPlatformProcess::GetCurrentProcessId(), WIDETEXT("_Startup")));
+}
+
+
+
+static ZenLockFileData ReadLockData(FUniqueBuffer&& FileBytes)
+{
+	if (ValidateCompactBinary(FileBytes, ECbValidateMode::Default) == ECbValidateError::None)
+	{
+		FCbObject LockObject(FileBytes.MoveToShared());
+		int32 ProcessId = LockObject["pid"].AsInt32();
+		FUtf8StringView DataDir = LockObject["data"].AsString();
+		int32 EffectivePort = LockObject["port"].AsInt32();
+		FCbObjectId SessionId = LockObject["session_id"].AsObjectId();
+		bool IsReady = LockObject["ready"].AsBool();
+		bool IsValid = ProcessId > 0 && EffectivePort > 0 && EffectivePort <= 0xffff;
+		return ZenLockFileData{ static_cast<uint32>(ProcessId), FString(DataDir), static_cast<uint16>(EffectivePort), SessionId, IsReady, IsValid };
+	}
+	return ZenLockFileData{ 0, {}, 0, {}, false };
+}
+
+bool
+ZenLockFileData::IsLockFileLocked(const TCHAR* FileName, bool bAttemptCleanUp)
+{
+#if PLATFORM_WINDOWS
+	if (bAttemptCleanUp)
+	{
+		IFileManager::Get().Delete(FileName, false, false, true);
+	}
+	return IFileManager::Get().FileExists(FileName);
+#elif PLATFORM_UNIX || PLATFORM_MAC
+	TAnsiStringBuilder<256> LockFilePath;
+	LockFilePath << FileName;
+	int32 Fd = open(LockFilePath.ToString(), O_RDONLY);
+	if (Fd < 0)
+	{
+		return false;
+	}
+
+	int32 LockRet = flock(Fd, LOCK_EX | LOCK_NB);
+	if (LockRet < 0)
+	{
+		close(Fd);
+		return errno == EWOULDBLOCK || errno == EAGAIN;
+	}
+
+	// Consider the lock file as orphaned if we we managed to claim the lock for
+	// it. Might as well delete it while we own it.
+	unlink(LockFilePath.ToString());
+
+	flock(Fd, LOCK_UN);
+	close(Fd);
+
+	return false;
+#endif
+}
+
+
+ZenLockFileData
+ZenLockFileData::ReadCbLockFile(const TCHAR* FileName)
+{
+#if PLATFORM_WINDOWS
+	// Windows specific lock reading path
+	// Uses share flags that are unique to windows to allow us to read file contents while the file may be open for write AND delete by another process (zenserver).
+
+	uint32 Access = GENERIC_READ;
+	uint32 WinFlags = FILE_SHARE_READ | FILE_SHARE_WRITE | FILE_SHARE_DELETE;
+	uint32 Create = OPEN_EXISTING;
+
+	TStringBuilder<MAX_PATH> FullFileNameBuilder;
+	FPathViews::ToAbsolutePath(FileName, FullFileNameBuilder);
+	for (TCHAR& Char : MakeArrayView(FullFileNameBuilder))
+	{
+		if (Char == TEXT('/'))
+		{
+			Char = TEXT('\\');
+		}
+	}
+	if (FullFileNameBuilder.Len() >= MAX_PATH)
+	{
+		FullFileNameBuilder.Prepend(TEXTVIEW("\\\\?\\"));
+	}
+	HANDLE Handle = CreateFileW(FullFileNameBuilder.ToString(), Access, WinFlags, NULL, Create, FILE_ATTRIBUTE_NORMAL, NULL);
+	if (Handle != INVALID_HANDLE_VALUE)
+	{
+		ON_SCOPE_EXIT{ CloseHandle(Handle); };
+		LARGE_INTEGER LI;
+		if (GetFileSizeEx(Handle, &LI))
+		{
+			checkf(LI.QuadPart == LI.u.LowPart, TEXT("Lock file exceeds supported 2GB limit."));
+			int32 FileSize32 = LI.u.LowPart;
+			FUniqueBuffer FileBytes = FUniqueBuffer::Alloc(FileSize32);
+			DWORD ReadBytes = 0;
+			if (ReadFile(Handle, FileBytes.GetData(), FileSize32, &ReadBytes, NULL) && (ReadBytes == FileSize32))
+			{
+				return ReadLockData(std::move(FileBytes));
+			}
+		}
+	}
+	return {};
+#elif PLATFORM_UNIX || PLATFORM_MAC
+	TAnsiStringBuilder<256> LockFilePath;
+	LockFilePath << FileName;
+	int32 Fd = open(LockFilePath.ToString(), O_RDONLY);
+	if (Fd < 0)
+	{
+		return {};
+	}
+
+	// If we can claim the lock then it's an orphaned lock file and should be
+	// ignored. Not ideal as there's a period of time when the lock can be
+	// held unncessarily.
+	int32 LockRet = flock(Fd, LOCK_EX | LOCK_NB);
+	if (LockRet >= 0)
+	{
+		unlink(LockFilePath.ToString());
+		flock(Fd, LOCK_UN);
+		close(Fd);
+		return {};
+	}
+
+	if (errno != EWOULDBLOCK && errno != EAGAIN)
+	{
+		close(Fd);
+		return {};
+	}
+
+	struct stat Stat;
+	fstat(Fd, &Stat);
+	uint64 FileSize = uint64(Stat.st_size);
+
+	FUniqueBuffer FileBytes = FUniqueBuffer::Alloc(FileSize);
+	if (read(Fd, FileBytes.GetData(), FileSize) == FileSize)
+	{
+		close(Fd);
+		return ReadLockData(std::move(FileBytes));
+	}
+
+	close(Fd);
+	return {};
+#endif
+}
diff --git a/Engine/Source/Developer/Zen/Private/ZenServerState.h b/Engine/Source/Developer/Zen/Private/ZenServerState.h
new file mode 100644
--- /dev/null
+++ b/Engine/Source/Developer/Zen/Private/ZenServerState.h
@@ -0,0 +1,105 @@
+// Copyright Epic Games, Inc. All Rights Reserved.
+
+#pragma once
+
+#include "Containers/StringFwd.h"
+#include "Serialization/CompactBinary.h"
+#include "Containers/AnsiString.h"
+
+class ZenServerState
+{
+public:
+	ZenServerState(bool ReadOnly);
+	~ZenServerState();
+
+	struct ZenServerEntry
+	{
+		// This matches the structure found in the Zen server
+		// https://github.com/EpicGames/zen/blob/main/zenutil/include/zenutil/zenserverprocess.h#L91
+		//
+		std::atomic<uint32> Pid;
+		std::atomic<uint16> DesiredListenPort;
+		std::atomic<uint16> Flags;
+		uint8				SessionId[12];
+		std::atomic<uint32> SponsorPids[8];
+		std::atomic<uint16> EffectiveListenPort;
+		uint8				Padding[10];
+
+		enum class FlagsEnum : uint16
+		{
+			kShutdownPlease = 1 << 0,
+			kIsReady = 1 << 1,
+		};
+
+		bool AddSponsorProcess(uint32 PidToAdd);
+	};
+	static_assert(sizeof(ZenServerEntry) == 64);
+
+	const ZenServerEntry* LookupByDesiredListenPort(int DesiredListenPort) const;
+	ZenServerEntry* LookupByDesiredListenPort(int DesiredListenPort);
+	const ZenServerEntry* LookupByEffectiveListenPort(int EffectiveListenPort) const;
+	ZenServerEntry* LookupByEffectiveListenPort(int EffectiveListenPort);
+	const ZenServerEntry* LookupByPid(uint32 Pid) const;
+
+	static bool IsProcessRunning(uint32 Pid);
+	static bool Terminate(uint32 Pid);
+	static bool FindRunningProcessId(const TCHAR* ExecutablePath, uint32* OutPid);
+
+private:
+	const ZenServerEntry* LookupByDesiredListenPortInternal(int DesiredListenPort) const;
+	const ZenServerEntry* LookupByEffectiveListenPortInternal(int EffectiveListenPort) const;
+	void* m_hMapFile = nullptr;
+	ZenServerEntry* m_Data = nullptr;
+	int				m_MaxEntryCount = 65536 / sizeof(ZenServerEntry);
+	bool			m_IsReadOnly = true;
+};
+
+class ZenSharedEvent
+{
+public:
+	ZenSharedEvent() = delete;
+	ZenSharedEvent(const ZenSharedEvent&) = delete;
+	ZenSharedEvent(ZenSharedEvent&&) = delete;
+	ZenSharedEvent& operator =(const ZenSharedEvent&) = delete;
+	ZenSharedEvent& operator =(ZenSharedEvent&&) = delete;
+
+	explicit ZenSharedEvent(FStringView EventName);
+	~ZenSharedEvent();
+
+	bool Create();
+	bool Exists() const;
+	bool Open();
+	bool Wait(int TimeoutMs = -1);
+	bool Set();
+	void Close();
+
+	static FString GetShutdownEventName(uint16 EffectiveListenPort);
+	static FString GetStartupEventName();
+
+private:
+	FString m_EventName;
+
+#if PLATFORM_WINDOWS
+	FString GetFullEventName() const;
+	void* m_EventHandle = NULL;
+#elif PLATFORM_UNIX || PLATFORM_MAC
+	FAnsiString GetEventPath() const;
+	int m_Fd = -1;
+	int m_Semaphore = -1;
+	FAnsiString m_EventPath;
+#endif
+};
+
+struct ZenLockFileData
+{
+	uint32 ProcessId = 0;
+	FString DataDir;
+	uint16 EffectivePort = 0;
+	FCbObjectId SessionId;
+	bool IsReady = false;
+	bool IsValid = false;
+
+	static bool IsLockFileLocked(const TCHAR* FileName, bool bAttemptCleanUp = false);
+	static ZenLockFileData ReadCbLockFile(const TCHAR* FileName);
+};
+
