Add missing check if zenserver process is running when we time out waiting for it
Fixes problem with CL 36532666

diff --git a/Engine/Source/Developer/Zen/Private/ZenServerInterface.cpp b/Engine/Source/Developer/Zen/Private/ZenServerInterface.cpp
--- a/Engine/Source/Developer/Zen/Private/ZenServerInterface.cpp
+++ b/Engine/Source/Developer/Zen/Private/ZenServerInterface.cpp
@@ -1883,30 +1883,34 @@ FZenServiceInstance::AutoLaunch(const FServiceAutoLaunchSettings& InSettings, FS
 				}
 
 				double ZenWaitDuration = FPlatformTime::ToSeconds64(FPlatformTime::Cycles64() - ZenWaitStartTime);
-				if (ZenWaitDuration >= 5.0)
+				if (ZenWaitDuration >= 10.0)
 				{
 					if (DurationPhase == EWaitDurationPhase::Short)
 					{
-#if !IS_PROGRAM
-						if (!FApp::IsUnattended() && !IsRunningCommandlet() && !GIsRunningUnattendedScript)
+						if (!FPlatformProcess::IsProcRunning(Proc))
 						{
-							FText ZenLaunchFailurePromptTitle = NSLOCTEXT("Zen", "Zen_LaunchFailurePromptTitle", "Failed to launch");
-
-							FFormatNamedArguments FormatArguments;
-							FString LogFilePath = FPaths::Combine(InSettings.DataPath, TEXT("logs"), TEXT("zenserver.log"));
-							FPaths::MakePlatformFilename(LogFilePath);
-							FormatArguments.Add(TEXT("LogFilePath"), FText::FromString(LogFilePath));
-							FText ZenLaunchFailurePromptText = FText::Format(NSLOCTEXT("Zen", "Zen_LaunchFailurePromptText", "Unreal Zen Storage Server failed to launch. Please check the ZenServer log file for details:\n{LogFilePath}"), FormatArguments);
-							FPlatformMisc::MessageBoxExt(EAppMsgType::Ok, *ZenLaunchFailurePromptText.ToString(), *ZenLaunchFailurePromptTitle.ToString());
-							break;
-						}
-						else
+#if !IS_PROGRAM
+							if (!FApp::IsUnattended() && !IsRunningCommandlet() && !GIsRunningUnattendedScript)
+							{
+								FText ZenLaunchFailurePromptTitle = NSLOCTEXT("Zen", "Zen_LaunchFailurePromptTitle", "Failed to launch");
+
+								FFormatNamedArguments FormatArguments;
+								FString LogFilePath = FPaths::Combine(InSettings.DataPath, TEXT("logs"), TEXT("zenserver.log"));
+								FPaths::MakePlatformFilename(LogFilePath);
+								FormatArguments.Add(TEXT("LogFilePath"), FText::FromString(LogFilePath));
+								FText ZenLaunchFailurePromptText = FText::Format(NSLOCTEXT("Zen", "Zen_LaunchFailurePromptText", "Unreal Zen Storage Server failed to launch. Please check the ZenServer log file for details:\n{LogFilePath}"), FormatArguments);
+								FPlatformMisc::MessageBoxExt(EAppMsgType::Ok, *ZenLaunchFailurePromptText.ToString(), *ZenLaunchFailurePromptTitle.ToString());
+								break;
+							}
+							else
 #endif
-						{
-							// Just log as there is no one to show a message
-							UE_LOG(LogZenServiceInstance, Warning, TEXT("Unreal Zen Storage Server did not launch in the expected duration"));
-							break;
+							{
+								// Just log as there is no one to show a message
+								UE_LOG(LogZenServiceInstance, Warning, TEXT("Unreal Zen Storage Server did not launch in the expected duration"));
+								break;
+							}
 						}
+
 						// Note that the dialog may not show up when zenserver is needed early in the launch cycle, but this will at least ensure
 						// the splash screen is refreshed with the appropriate text status message.
 						WaitForZenReadySlowTask.MakeDialog(true, false);
