Make sure we are exclusively trying to startup zenserver at autolaunch

See https://epicgames.lightning.force.com/lightning/r/Case/500QP00000Lm1aWYAR/view

diff --git a/Engine/Source/Developer/Zen/Private/ZenServerInterface.cpp b/Engine/Source/Developer/Zen/Private/ZenServerInterface.cpp
--- a/Engine/Source/Developer/Zen/Private/ZenServerInterface.cpp
+++ b/Engine/Source/Developer/Zen/Private/ZenServerInterface.cpp
@@ -1398,7 +1398,7 @@ FZenServiceInstance::TryRecovery()
 
 	if (TimespanSinceLastRecovery > MinimumDurationSinceLastRecovery)
 	{
-		static FSystemWideCriticalSection RecoveryCriticalSection(TEXT("ZenServerRecovery"), MaximumWaitForHealth);
+		FSystemWideCriticalSection RecoveryCriticalSection(TEXT("ZenServerRecovery"), MaximumWaitForHealth);
 		if (!RecoveryCriticalSection.IsValid())
 		{
 			// A recovery is already in progress but did not complete in time, we assume we failed and let recovery continue on a different thread
@@ -1509,24 +1509,35 @@ FZenServiceInstance::Initialize()
 		{
 			int LaunchAttempts = 0;
 			const FTimespan MaximumWaitForHealth = FTimespan::FromSeconds(20);
+
 			FDateTime StartedWaitingForHealth = FDateTime::UtcNow();
 			while (true)
 			{
-				++LaunchAttempts;
-				bHasLaunchedLocal = AutoLaunch(Settings.SettingsVariant.Get<FServiceAutoLaunchSettings>(), *ExecutableInstallPath, HostName, Port);
-				if (bHasLaunchedLocal)
 				{
-					const ZenServerState State(/*ReadOnly*/true);
-					const ZenServerState::ZenServerEntry* RunningEntry = State.LookupByEffectiveListenPort(Port);
-					if (RunningEntry != nullptr)
+					FSystemWideCriticalSection RecoveryCriticalSection(TEXT("ZenServerRecovery"), MaximumWaitForHealth);
+					if (!RecoveryCriticalSection.IsValid())
 					{
-						AutoLaunchedPid = RunningEntry->Pid.load(std::memory_order_relaxed);
+						UE_LOG(LogZenServiceInstance, Warning, TEXT("Local ZenServer AutoLaunch initialization timed out waiting for other process to complete startup/recovery"));
 					}
-					AutoLaunchedPort = Port;
-					bIsRunningLocally = true;
-					if (IsServiceReady())
+					else
 					{
-						break;
+						++LaunchAttempts;
+						bHasLaunchedLocal = AutoLaunch(Settings.SettingsVariant.Get<FServiceAutoLaunchSettings>(), *ExecutableInstallPath, HostName, Port);
+						if (bHasLaunchedLocal)
+						{
+							const ZenServerState State(/*ReadOnly*/true);
+							const ZenServerState::ZenServerEntry* RunningEntry = State.LookupByEffectiveListenPort(Port);
+							if (RunningEntry != nullptr)
+							{
+								AutoLaunchedPid = RunningEntry->Pid.load(std::memory_order_relaxed);
+							}
+							AutoLaunchedPort = Port;
+							bIsRunningLocally = true;
+							if (IsServiceReady())
+							{
+								break;
+							}
+						}
 					}
 				}
 
