Allow zenserver to be installed using a "link" to the installed UE executables instead of copying to the %PROGRAMDATA% folder which may have restrictions preventing the zenserver executable from running.

Selecting "link" or "copy" is done automatically by checking if UE is running from an installed place or if it is a dev/p4 stream place, you can override this with the option -ZenAutoLaunchInstallMode.
The ZenAutoLaunchInstallMode defaults to "auto" where it tries to detect if it is an installed engine or not. "link" forces the use of a link file and "copy" forces the copy-install option which is the legacy behavior.

Additional fixes:
FMacSystemWideCriticalSection which could fail to aquire the critical section causing launch to fail  now uses flock in the same manner as FUnixSystemWideCriticalSection which resolves the issues
When checking if processes are running on Unix/Mac we now detect zombie processes properly

diff --git a/Engine/Source/Developer/Zen/Internal/ZenVersion.h b/Engine/Source/Developer/Zen/Internal/ZenVersion.h
--- a/Engine/Source/Developer/Zen/Internal/ZenVersion.h
+++ b/Engine/Source/Developer/Zen/Internal/ZenVersion.h
@@ -19,6 +19,7 @@ struct FZenVersion
 	bool TryParse(const TCHAR* InString);
 	FString ToString(bool bDetailed = true) const;
 	bool operator<(FZenVersion& Other) const;
+	operator bool() const;
 };
 
 } // namespace UE::Zen
diff --git a/Engine/Source/Developer/Zen/Private/ZenServerInterface.cpp b/Engine/Source/Developer/Zen/Private/ZenServerInterface.cpp
--- a/Engine/Source/Developer/Zen/Private/ZenServerInterface.cpp
+++ b/Engine/Source/Developer/Zen/Private/ZenServerInterface.cpp
@@ -48,74 +48,261 @@ namespace UE::Zen
 
 DEFINE_LOG_CATEGORY_STATIC(LogZenServiceInstance, Log, All);
 
-static FZenVersion
-GetZenVersion(const FString& UtilityPath, const FString& ServicePath, const FString& VersionCachePath)
+struct FZenServiceLink
 {
-	IFileManager& FileManager = IFileManager::Get();
-	FDateTime UtilityExecutableModificationTime = FileManager.GetTimeStamp(*UtilityPath);
-	FDateTime ServiceExecutableModificationTime = FileManager.GetTimeStamp(*ServicePath);
+	FString ServicePath;
+	FString UtilityPath;
+	FZenVersion Version;
+
+	operator bool() const
+	{
+		return !ServicePath.IsEmpty() && !UtilityPath.IsEmpty() && Version;
+	}
 
-	auto GetFallbackVersion = [&UtilityExecutableModificationTime, &ServiceExecutableModificationTime]()
+	static FZenServiceLink Read(const FString& Filename)
 	{
-		FZenVersion FallbackVersion;
-		if (UtilityExecutableModificationTime > ServiceExecutableModificationTime)
+		FString JsonText;
+		if (FFileHelper::LoadFileToString(JsonText, *Filename))
 		{
-			FallbackVersion.Details = UtilityExecutableModificationTime.ToString();
-			return FallbackVersion;
+			TSharedPtr<FJsonObject> JsonObject;
+			TSharedRef<TJsonReader<TCHAR>> Reader = TJsonReaderFactory<TCHAR>::Create(JsonText);
+			if (FJsonSerializer::Deserialize(Reader, JsonObject) && JsonObject.IsValid())
+			{
+				FString ServicePath = JsonObject->Values.FindRef(TEXT("ServicePath"))->AsString();
+				FString UtilityPath = JsonObject->Values.FindRef(TEXT("UtilityPath"))->AsString();
+				auto VersionObject = JsonObject->Values.FindRef(TEXT("Version"))->AsObject();
+				if (VersionObject)
+				{
+					uint32_t MajorVersion = static_cast<uint32_t>(VersionObject->TryGetField(TEXT("Major"))->AsNumber());
+					uint32_t MinorVersion = static_cast<uint32_t>(VersionObject->TryGetField(TEXT("Minor"))->AsNumber());
+					uint32_t PatchVersion = static_cast<uint32_t>(VersionObject->TryGetField(TEXT("Patch"))->AsNumber());
+					FString Details = VersionObject->TryGetField(TEXT("Details"))->AsString();
+					return FZenServiceLink{
+						.ServicePath = ServicePath,
+						.UtilityPath = UtilityPath,
+						.Version = FZenVersion{
+							.MajorVersion = MajorVersion,
+							.MinorVersion = MinorVersion,
+							.PatchVersion = PatchVersion,
+							.Details = Details}
+					};
+				}
+			}
 		}
-		FallbackVersion.Details = ServiceExecutableModificationTime.ToString();
-		return FallbackVersion;
-	};
+		return {};
+	}
 
-	FDateTime VersionCacheModificationTime = FileManager.GetTimeStamp(*VersionCachePath);
-	bool VersionCacheIsOlderThanUtilityExecutable = VersionCacheModificationTime < UtilityExecutableModificationTime;
-	bool VersionCacheIsOlderThanServerExecutable = VersionCacheModificationTime < ServiceExecutableModificationTime;
-	bool VersionCacheIsOutOfDate = VersionCacheIsOlderThanUtilityExecutable || VersionCacheIsOlderThanServerExecutable;
-	FString VersionFileContents;
-	FZenVersion ComparableVersion;
-	if (VersionCacheIsOutOfDate ||
-		!FFileHelper::LoadFileToString(VersionFileContents, *VersionCachePath) ||
-		!ComparableVersion.TryParse(*VersionFileContents))
+	static bool Write(const FZenServiceLink& Link, const FString& Filename)
 	{
-		FString AbsoluteUtilityPath = FPaths::ConvertRelativePathToFull(UtilityPath);
-		FMonitoredProcess MonitoredUtilityProcess(AbsoluteUtilityPath, TEXT("version --detailed"), FPaths::GetPath(UtilityPath), true);
+		FString JsonTcharText;
+		TSharedRef<TJsonWriter<TCHAR, TPrettyJsonPrintPolicy<TCHAR>>> Writer = TJsonWriterFactory<TCHAR, TPrettyJsonPrintPolicy<TCHAR>>::Create(&JsonTcharText);
+		Writer->WriteObjectStart();
 
-		bool bLaunched = MonitoredUtilityProcess.Launch();
-		checkf(bLaunched, TEXT("Failed to launch zen utility to gather version data: '%s'."), *UtilityPath);
-		if (!bLaunched)
-		{
-			return GetFallbackVersion();
-		}
+		Writer->WriteValue(TEXT("ServicePath"), Link.ServicePath);
+		Writer->WriteValue(TEXT("UtilityPath"), Link.UtilityPath);
+
+		Writer->WriteObjectStart(TEXT("Version"));
+		Writer->WriteValue(TEXT("Major"), Link.Version.MajorVersion);
+		Writer->WriteValue(TEXT("Minor"), Link.Version.MinorVersion);
+		Writer->WriteValue(TEXT("Patch"), Link.Version.PatchVersion);
+		Writer->WriteValue(TEXT("Details"), Link.Version.Details);
+		Writer->WriteObjectEnd();
 
-		while (MonitoredUtilityProcess.Update())
+		Writer->WriteObjectEnd();
+		Writer->Close();
+
+		if (!FFileHelper::SaveStringToFile(JsonTcharText, *Filename))
 		{
-			FPlatformProcess::Sleep(0.1f);
-			if (MonitoredUtilityProcess.GetDuration().GetTotalSeconds() > 10)
-			{
-				MonitoredUtilityProcess.Cancel(true);
-				checkf(false, TEXT("Cancelled launch of zen utility for gathering version data: '%s'."), *UtilityPath);
-				return GetFallbackVersion();
-			}
+			return false;
 		}
 
-		FString OutputString = MonitoredUtilityProcess.GetFullOutputWithoutDelegate();
-		if (MonitoredUtilityProcess.GetReturnCode() != 0)
+		return true;
+	}
+};
+
+static FString
+GetLocalZenRootPath()
+{
+	return FPaths::ConvertRelativePathToFull(FPaths::Combine(FPlatformProcess::UserSettingsDir(), *FApp::GetEpicProductIdentifier(), TEXT("Common")) + TEXT("/"));
+}
+
+static FString
+GetServiceExecutableName()
+{
+	return
+#if PLATFORM_WINDOWS
+		TEXT("zenserver.exe")
+#else
+		TEXT("zenserver")
+#endif
+		;
+}
+
+static FString
+GetUtilityExecutableName()
+{
+	return
+#if PLATFORM_WINDOWS
+		TEXT("zen.exe")
+#else
+		TEXT("zen")
+#endif
+		;
+}
+
+static FString
+GetLocalInstallPath()
+{
+	return FPaths::ConvertRelativePathToFull(FPaths::Combine(GetLocalZenRootPath(), TEXT("Zen\\Install")));
+}
+
+static FString
+GetServiceLinkPath()
+{
+	return FPaths::Combine(GetLocalInstallPath(), TEXT("zen.link"));
+}
+
+static FString
+GetServiceCopyInstallPath()
+{
+	return FPaths::ConvertRelativePathToFull(FPaths::Combine(GetLocalInstallPath(), GetServiceExecutableName()));
+}
+
+static FString
+GetUtilityCopyInstallPath()
+{
+	return FPaths::ConvertRelativePathToFull(FPaths::Combine(GetLocalInstallPath(), GetUtilityExecutableName()));
+}
+
+static FString
+GetInstallVersionCachePath()
+{
+	FString InstallUtilityPath = GetUtilityCopyInstallPath();
+	FString InstallVersionCache = FPaths::SetExtension(InstallUtilityPath, TEXT("version"));
+	return InstallVersionCache;
+}
+
+static FString
+GetInTreeVersionCache()
+{
+	return FPaths::ConvertRelativePathToFull(FPaths::Combine(FPaths::EngineSavedDir(), TEXT("Zen"), TEXT("zen.version")));
+}
+
+static FString
+GetServiceRunContextPath()
+{
+	return FPaths::SetExtension(GetServiceCopyInstallPath(), TEXT(".runcontext"));
+}
+
+static FString
+GetInTreeUtilityPath()
+{
+	return FPaths::ConvertRelativePathToFull(FPlatformProcess::GenerateApplicationPath(TEXT("zen"), EBuildConfiguration::Development));
+}
+
+static FString
+GetInTreeServicePath()
+{
+	return FPaths::ConvertRelativePathToFull(FPlatformProcess::GenerateApplicationPath(TEXT("zenserver"), EBuildConfiguration::Development));
+}
+
+static FString
+GetInTreeCrashpadHandlerFilePath()
+{
+	return FPaths::ConvertRelativePathToFull(FPlatformProcess::GenerateApplicationPath(TEXT("crashpad_handler"), EBuildConfiguration::Development));
+}
+
+static FString
+GetInstallCrashpadHandlerFilePath(const FString& InTreePath)
+{
+	return FPaths::ConvertRelativePathToFull(FPaths::Combine(GetLocalInstallPath(), FString(FPathViews::GetCleanFilename(InTreePath))));
+}
+
+
+static bool GetZenVersion(const FString& UtilityPath, const FString& ServicePath, FZenVersion& OutVersion)
+{
+	FString AbsoluteUtilityPath = FPaths::ConvertRelativePathToFull(UtilityPath);
+	FMonitoredProcess MonitoredUtilityProcess(AbsoluteUtilityPath, TEXT("version --detailed"), FPaths::GetPath(UtilityPath), true);
+
+	bool bLaunched = MonitoredUtilityProcess.Launch();
+	if (!bLaunched)
+	{
+		UE_LOG(LogZenServiceInstance, Warning, TEXT("Failed to launch zen utility to gather version data: '%s'."), *UtilityPath);
+		return false;
+	}
+
+	while (MonitoredUtilityProcess.Update())
+	{
+		FPlatformProcess::Sleep(0.1f);
+		if (MonitoredUtilityProcess.GetDuration().GetTotalSeconds() > 10)
 		{
-			checkf(false, TEXT("Unexpected return code after launch of zen utility for gathering version data: '%s' (%d). Output: '%s'"), *UtilityPath, MonitoredUtilityProcess.GetReturnCode(), *OutputString);
-			return GetFallbackVersion();
+			MonitoredUtilityProcess.Cancel(true);
+			UE_LOG(LogZenServiceInstance, Warning, TEXT("Cancelled launch of zen utility for gathering version data: '%s'."), *UtilityPath);
+			return false;
 		}
+	}
+
+	FString OutputString = MonitoredUtilityProcess.GetFullOutputWithoutDelegate();
+	if (MonitoredUtilityProcess.GetReturnCode() != 0)
+	{
+		UE_LOG(LogZenServiceInstance, Warning, TEXT("Unexpected return code after launch of zen utility for gathering version data: '%s' (%d). Output: '%s'"), *UtilityPath, MonitoredUtilityProcess.GetReturnCode(), *OutputString);
+		return false;
+	}
+
+	FString VersionOutputString = OutputString.TrimStartAndEnd();
+
+	if (!OutVersion.TryParse(*VersionOutputString))
+	{
+		UE_LOG(LogZenServiceInstance, Warning, TEXT("Invalid version information after launch of zen utility for gathering version data: '%s' (`%s`)"), *UtilityPath, *VersionOutputString);
+		return false;
+	}
+	return true;
+}
 
-		FString VersionOutputString = OutputString.TrimStartAndEnd();
-		if (!ComparableVersion.TryParse(*VersionOutputString))
+static FZenVersion
+GetZenVersion(const FString& UtilityPath, const FString& ServicePath, const FString& VersionCachePath)
+{
+	IFileManager& FileManager = IFileManager::Get();
+	FDateTime UtilityExecutableModificationTime = FileManager.GetTimeStamp(*UtilityPath);
+	FDateTime ServiceExecutableModificationTime = FileManager.GetTimeStamp(*ServicePath);
+
+	FDateTime VersionCacheModificationTime = FileManager.GetTimeStamp(*VersionCachePath);
+	bool VersionCacheIsOlderThanUtilityExecutable = VersionCacheModificationTime < UtilityExecutableModificationTime;
+	bool VersionCacheIsOlderThanServerExecutable = VersionCacheModificationTime < ServiceExecutableModificationTime;
+	bool VersionCacheIsUpToDate = (!VersionCacheIsOlderThanUtilityExecutable) && (!VersionCacheIsOlderThanServerExecutable);
+	if (VersionCacheIsUpToDate)
+	{
+		FString VersionFileContents;
+		if (FFileHelper::LoadFileToString(VersionFileContents, *VersionCachePath))
 		{
-			checkf(false, TEXT("Invalid version information after launch of zen utility for gathering version data: '%s' (`%s`)"), *UtilityPath, *VersionOutputString);
-			return GetFallbackVersion();
+			FZenVersion CachedVersion;
+			if (CachedVersion.TryParse(*VersionFileContents))
+			{
+				return CachedVersion;
+			}
 		}
-
-		FFileHelper::SaveStringToFile(ComparableVersion.ToString(), *VersionCachePath);
 	}
+	auto GetFallbackVersion = [UtilityExecutableModificationTime, ServiceExecutableModificationTime]()
+		{
+			FZenVersion FallbackVersion;
+			if (UtilityExecutableModificationTime > ServiceExecutableModificationTime)
+			{
+				FallbackVersion.Details = UtilityExecutableModificationTime.ToString();
+				return FallbackVersion;
+			}
+			FallbackVersion.Details = ServiceExecutableModificationTime.ToString();
+			return FallbackVersion;
+		};
 
-	return ComparableVersion;
+
+	FZenVersion Version;
+	if (!GetZenVersion(UtilityPath, ServicePath, Version))
+	{
+		checkf(false, TEXT("Unable to determine version using zen utility executable path: '%s'."), *UtilityPath);
+		Version = GetFallbackVersion();
+	}
+	FFileHelper::SaveStringToFile(Version.ToString(), *VersionCachePath);
+	return Version;
 }
 
 static void
@@ -137,11 +324,8 @@ PromptUserToSyncInTreeVersion(const FString& ServerFilePath)
 }
 
 static bool
-IsInstallVersionOutOfDate(const FString& InTreeUtilityPath, const FString& InstallUtilityPath, const FString& InTreeServicePath, const FString& InstallServicePath, FString& OutInTreeVersionCache, FString& OutInstallVersionCache)
+IsInstallVersionOutOfDate(const FString& InTreeUtilityPath, const FString& InstallUtilityPath, const FString& InTreeServicePath, const FString& InstallServicePath, const FString& InTreeVersionCache, const FString& InstallVersionCache)
 {
-	OutInTreeVersionCache = FPaths::Combine(FPaths::EngineSavedDir(), TEXT("Zen"), TEXT("zen.version"));
-	OutInstallVersionCache = FPaths::SetExtension(InstallUtilityPath, TEXT("version"));
-
 	IFileManager& FileManager = IFileManager::Get();
 	if (!FileManager.FileExists(*InTreeUtilityPath) || !FileManager.FileExists(*InTreeServicePath))
 	{
@@ -151,7 +335,7 @@ IsInstallVersionOutOfDate(const FString& InTreeUtilityPath, const FString& Insta
 	}
 
 	// Always get the InTree utility path so cached version information is up to date
-	FZenVersion InTreeVersion = GetZenVersion(InTreeUtilityPath, InTreeServicePath, OutInTreeVersionCache);
+	FZenVersion InTreeVersion = GetZenVersion(InTreeUtilityPath, InTreeServicePath, InTreeVersionCache);
 	UE_LOG(LogZenServiceInstance, Log, TEXT("InTree version at '%s' is '%s'"), *InTreeServicePath, *InTreeVersion.ToString());
 
 	if (!FileManager.FileExists(*InstallUtilityPath) || !FileManager.FileExists(*InstallServicePath))
@@ -159,7 +343,7 @@ IsInstallVersionOutOfDate(const FString& InTreeUtilityPath, const FString& Insta
 		UE_LOG(LogZenServiceInstance, Log, TEXT("No installation found at '%s'"), *InstallServicePath);
 		return true;
 	}
-	FZenVersion InstallVersion = GetZenVersion(InstallUtilityPath, InstallServicePath, OutInstallVersionCache);
+	FZenVersion InstallVersion = GetZenVersion(InstallUtilityPath, InstallServicePath, InstallVersionCache);
 	UE_LOG(LogZenServiceInstance, Log, TEXT("Installed version at '%s' is '%s'"), *InstallServicePath, *InstallVersion.ToString());
 
 	if (InstallVersion < InTreeVersion)
@@ -178,7 +362,8 @@ IsInstallVersionOutOfDate(const FString& InTreeUtilityPath, const FString& Insta
 static bool
 AttemptFileCopyWithRetries(const TCHAR* Dst, const TCHAR* Src, double RetryDurationSeconds)
 {
-	uint32 CopyResult = IFileManager::Get().Copy(Dst, Src, true, true, false);
+	IFileManager& FileManager = IFileManager::Get();
+	uint32 CopyResult = FileManager.Copy(Dst, Src, true, true, false);
 	uint64 CopyWaitStartTime = FPlatformTime::Cycles64();
 	while (CopyResult != COPY_OK)
 	{
@@ -191,7 +376,7 @@ AttemptFileCopyWithRetries(const TCHAR* Dst, const TCHAR* Src, double RetryDurat
 		{
 			break;
 		}
-		CopyResult = IFileManager::Get().Copy(Dst, Src, true, true, false);
+		CopyResult = FileManager.Copy(Dst, Src, true, true, false);
 	}
 	if (CopyResult == COPY_OK)
 	{
@@ -201,6 +386,32 @@ AttemptFileCopyWithRetries(const TCHAR* Dst, const TCHAR* Src, double RetryDurat
 	return false;
 }
 
+static bool
+AttemptFileDeleteWithRetries(const TCHAR* Path, double RetryDurationSeconds)
+{
+	IFileManager& FileManager = IFileManager::Get();
+	bool DeleteResult = FileManager.Delete(Path, false, false, true);
+	uint64 DeleteWaitStartTime = FPlatformTime::Cycles64();
+	while (!DeleteResult)
+	{
+		double DeleteWaitDuration = FPlatformTime::ToSeconds64(FPlatformTime::Cycles64() - DeleteWaitStartTime);
+		if (DeleteWaitDuration < RetryDurationSeconds)
+		{
+			FPlatformProcess::Sleep(0.01f);
+		}
+		else
+		{
+			break;
+		}
+		DeleteResult = FileManager.Delete(Path, false, false, true);
+	}
+	if (DeleteResult)
+	{
+		return true;
+	}
+	return false;
+}
+
 static void EnsureEditorSettingsConfigLoaded()
 {
 #if !WITH_EDITOR
@@ -275,12 +486,6 @@ DetermineLocalDataCachePath(const TCHAR* ConfigSection, FString& DataPath)
 	}
 }
 
-static FString
-GetLocalZenRootPath()
-{
-	return FPaths::ConvertRelativePathToFull(FPaths::Combine(FPlatformProcess::UserSettingsDir(), *FApp::GetEpicProductIdentifier(), TEXT("Common")) + TEXT("/"));
-}
-
 static bool
 DetermineDataPath(const TCHAR* ConfigSection, FString& DataPath, bool& bHasInvalidPathConfigurations, bool& bIsDefaultDataPath)
 {
@@ -524,6 +729,40 @@ static void ReadCbField(FCbFieldView Field, UE::Zen::FZenCIDStats& OutValue)
 	ReadCbField(ObjectView["size"], OutValue.Size);
 }
 
+static FServiceAutoLaunchSettings::EInstallMode ZenGetInstallModeFromString(const FString& InstallMode)
+{
+	if (InstallMode.IsEmpty() || FCString::Stricmp(*InstallMode, TEXT("auto")) == 0)
+	{
+		return FApp::IsEngineInstalled() ? FServiceAutoLaunchSettings::EInstallMode::Link : FServiceAutoLaunchSettings::EInstallMode::Copy;
+	}
+	else if (FCString::Stricmp(*InstallMode, TEXT("copy")) == 0)
+	{
+		return FServiceAutoLaunchSettings::EInstallMode::Copy;
+	}
+	else if (FCString::Stricmp(*InstallMode, TEXT("link")) == 0)
+	{
+		return FServiceAutoLaunchSettings::EInstallMode::Link;
+	}
+	else
+	{
+		checkf(false, TEXT("Invalid zenserver install mode: {InstallMode}"), *InstallMode);
+		return FServiceAutoLaunchSettings::EInstallMode::Copy;
+	}
+}
+
+static FString ZenGetInstallModeToString(FServiceAutoLaunchSettings::EInstallMode InstallMode)
+{
+	switch (InstallMode)
+	{
+	case FServiceAutoLaunchSettings::EInstallMode::Copy:
+		return TEXT("copy");
+	case FServiceAutoLaunchSettings::EInstallMode::Link:
+		return TEXT("link");
+	}
+	checkf(false, TEXT("Invalid zenserver install mode: {%d}"), static_cast<int>(InstallMode));
+	return "";
+}
+
 bool
 FServiceSettings::ReadFromConfig()
 {
@@ -589,6 +828,11 @@ FServiceSettings::ReadFromConfig()
 			GConfig->GetBool(AutoLaunchConfigSection, TEXT("LimitProcessLifetime"), AutoLaunchSettings.bLimitProcessLifetime, GEngineIni);
 			ApplyProcessLifetimeOverride(AutoLaunchSettings.bLimitProcessLifetime);
 			GConfig->GetBool(AutoLaunchConfigSection, TEXT("AllowRemoteNetworkService"), AutoLaunchSettings.bAllowRemoteNetworkService, GEngineIni);
+			FString InstallMode;
+			if (GConfig->GetString(AutoLaunchConfigSection, TEXT("InstallMode"), InstallMode, GEngineIni))
+			{
+				AutoLaunchSettings.InstallMode = ZenGetInstallModeFromString(InstallMode);
+			}
 			EnsureEditorSettingsConfigLoaded();
 			GConfig->GetBool(TEXT("/Script/UnrealEd.CrashReportsPrivacySettings"), TEXT("bSendUnattendedBugReports"), AutoLaunchSettings.bSendUnattendedBugReports, GEditorSettingsIni);
 		}
@@ -628,6 +872,7 @@ FServiceSettings::ReadFromCompactBinary(FCbFieldView Field)
 				AutoLaunchSettings.bAllowRemoteNetworkService = AutoLaunchSettingsObject["AllowRemoteNetworkService"].AsBool();
 				AutoLaunchSettings.bSendUnattendedBugReports = AutoLaunchSettingsObject["SendUnattendedBugReports"].AsBool();
 				AutoLaunchSettings.bIsDefaultSharedRunContext = AutoLaunchSettingsObject["IsDefaultSharedRunContext"].AsBool(AutoLaunchSettings.bIsDefaultSharedRunContext);
+				AutoLaunchSettings.InstallMode = ZenGetInstallModeFromString(FString(AutoLaunchSettingsObject["InstallMode"].AsString()));
 			}
 		}
 	}
@@ -689,6 +934,7 @@ FServiceSettings::WriteToCompactBinary(FCbWriter& Writer) const
 		Writer << "AllowRemoteNetworkService" << AutoLaunchSettings.bAllowRemoteNetworkService;
 		Writer << "SendUnattendedBugReports" << AutoLaunchSettings.bSendUnattendedBugReports;
 		Writer << "IsDefaultSharedRunContext" << AutoLaunchSettings.bIsDefaultSharedRunContext;
+		Writer << "InstallMode" << ZenGetInstallModeToString(AutoLaunchSettings.InstallMode);
 		Writer.EndObject();
 	}
 	else
@@ -743,41 +989,6 @@ FServiceSettings::TryApplyAutoLaunchOverride()
 uint16 FZenServiceInstance::AutoLaunchedPort = 0;
 uint32 FZenServiceInstance::AutoLaunchedPid = 0;
 
-FString GetLocalServiceExecutableName()
-{
-	return
-#if PLATFORM_WINDOWS
-		TEXT("zenserver.exe")
-#else
-		TEXT("zenserver")
-#endif
-		;
-}
-
-FString
-GetLocalInstallPath()
-{
-	return FPaths::ConvertRelativePathToFull(FPaths::Combine(GetLocalZenRootPath(), TEXT("Zen\\Install")));
-}
-
-FString
-GetLocalServiceInstallPath()
-{
-	return FPaths::ConvertRelativePathToFull(FPaths::Combine(GetLocalInstallPath(), GetLocalServiceExecutableName()));
-}
-
-FString
-GetLocalInstallUtilityPath()
-{
-	return FPaths::ConvertRelativePathToFull(FPaths::Combine(GetLocalInstallPath(),
-#if PLATFORM_WINDOWS
-		TEXT("zen.exe")
-#else
-		TEXT("zen")
-#endif
-	));
-}
-
 static bool
 IsZenProcessUsingEffectivePort(uint16 EffectiveListenPort)
 {
@@ -850,7 +1061,7 @@ static bool ShutDownZenServerProcessExecutable(const FString& ExecutablePath, do
 	{
 		UE_LOG(LogZenServiceInstance, Display, TEXT("Attempting to shut down of zenserver executable '%s' process with pid %d"), *ExecutablePath, Pid);
 		double ZenShutdownWaitDuration = FPlatformTime::ToSeconds64(FPlatformTime::Cycles64() - ZenShutdownWaitStartTime);
-		if (ShutdownZenServerProcess(Pid, ZenShutdownWaitDuration))
+		if (ShutdownZenServerProcess(Pid, MaximumWaitDurationSeconds - ZenShutdownWaitDuration))
 		{
 			return true;
 		}
@@ -883,7 +1094,7 @@ static bool ShutDownZenServerProcessLockingDataDir(const FString& DataPath, doub
 				return true;
 			}
 			uint32_t Pid = 0;
-			if (!ZenServerState::FindRunningProcessId(*GetLocalServiceExecutableName(), & Pid))
+			if (!ZenServerState::FindRunningProcessId(*GetServiceExecutableName(), &Pid))
 			{
 				if (!ZenLockFileData::IsLockFileLocked(*LockFilePath, true))
 				{
@@ -895,7 +1106,7 @@ static bool ShutDownZenServerProcessLockingDataDir(const FString& DataPath, doub
 			}
 			UE_LOG(LogZenServiceInstance, Warning, TEXT("Found locked but invalid lock file at '%s', attempting shut down of zenserver process with pid %d"), *LockFilePath, Pid);
 			double ZenShutdownWaitDuration = FPlatformTime::ToSeconds64(FPlatformTime::Cycles64() - ZenShutdownWaitStartTime);
-			if (!ShutdownZenServerProcess(Pid, ZenShutdownWaitDuration))
+			if (!ShutdownZenServerProcess(Pid, MaximumWaitDurationSeconds - ZenShutdownWaitDuration))
 			{
 				break;
 			}
@@ -946,7 +1157,7 @@ static bool ShutDownZenServerProcessLockingDataDir(const FString& DataPath, doub
 			return true;
 		}
 		uint32_t Pid = 0;
-		if (!ZenServerState::FindRunningProcessId(*GetLocalServiceExecutableName(), &Pid))
+		if (!ZenServerState::FindRunningProcessId(*GetServiceExecutableName(), &Pid))
 		{
 			if (!ZenLockFileData::IsLockFileLocked(*LockFilePath, true))
 			{
@@ -958,7 +1169,7 @@ static bool ShutDownZenServerProcessLockingDataDir(const FString& DataPath, doub
 		}
 		UE_LOG(LogZenServiceInstance, Warning, TEXT("Found locked but invalid lock file at '%s', attempting shut down of zenserver process with pid %d"), *LockFilePath, Pid);
 		double ZenShutdownWaitDuration = FPlatformTime::ToSeconds64(FPlatformTime::Cycles64() - ZenShutdownWaitStartTime);
-		if (!ShutdownZenServerProcess(Pid, ZenShutdownWaitDuration))
+		if (!ShutdownZenServerProcess(Pid, MaximumWaitDurationSeconds - ZenShutdownWaitDuration))
 		{
 			break;
 		}
@@ -1066,8 +1277,7 @@ Private::GetLocalDataCachePathOverride(FString& OutDataPath)
 bool
 TryGetLocalServiceRunContext(FZenLocalServiceRunContext& OutContext)
 {
-	FString LocalServiceRunContextPath = FPaths::SetExtension(GetLocalServiceInstallPath(), TEXT(".runcontext"));
-	return OutContext.ReadFromJsonFile(*LocalServiceRunContextPath);
+	return OutContext.ReadFromJsonFile(*GetServiceRunContextPath());
 }
 
 bool
@@ -1240,20 +1450,55 @@ StopLocalService(const TCHAR* DataPath, double MaximumWaitDurationSeconds)
 }
 
 FString
-GetLocalServiceInstallVersion(bool bDetailed)
+GetLocalServiceInstallPath()
 {
-	FString InstallUtilityPath = GetLocalInstallUtilityPath();
-	FString InstallVersionCache = FPaths::SetExtension(InstallUtilityPath, TEXT("version"));
+	if (FZenServiceLink Link = FZenServiceLink::Read(GetServiceLinkPath()); Link)
+	{
+		return Link.ServicePath;
+	}
+	else
+	{
+		return GetServiceCopyInstallPath();
+	}
+}
+
+FString
+GetLocalInstallUtilityPath()
+{
+	if (FZenServiceLink Link = FZenServiceLink::Read(GetServiceLinkPath()); Link)
+	{
+		return Link.UtilityPath;
+	}
+	else
+	{
+		return GetUtilityCopyInstallPath();
+	}
+}
 
+FString
+GetLocalServiceInstallVersion(bool bDetailed)
+{
 	IFileManager& FileManager = IFileManager::Get();
-	if (!FileManager.FileExists(*InstallUtilityPath))
+	if (FZenServiceLink Link = FZenServiceLink::Read(GetServiceLinkPath()); Link)
 	{
-		return FZenVersion().ToString(bDetailed);
+		if (!FileManager.FileExists(*Link.ServicePath))
+		{
+			return FZenVersion().ToString(bDetailed);
+		}
+		return Link.Version.ToString(bDetailed);
 	}
+	else
+	{
+		const FString ServicePath = GetServiceCopyInstallPath();
 
-	FZenVersion InstallVersion = GetZenVersion(InstallUtilityPath, GetLocalServiceInstallPath(), InstallVersionCache);
+		if (!FileManager.FileExists(*ServicePath))
+		{
+			return FZenVersion().ToString(bDetailed);
+		}
 
-	return InstallVersion.ToString(bDetailed);
+		FZenVersion InstallVersion = GetZenVersion(GetUtilityCopyInstallPath(), ServicePath, GetInstallVersionCachePath());
+		return InstallVersion.ToString(bDetailed);
+	}
 }
 
 static bool GIsDefaultServicePresent = false;
@@ -1389,6 +1634,9 @@ FZenServiceInstance::TryRecovery()
 		return false;
 	}
 
+	const FString ExecutablePath = GetLocalServiceInstallPath();
+	const FString ExecutionContextFilePath = GetServiceRunContextPath();
+
 	static std::atomic<int64> LastRecoveryTicks;
 	static bool bLastRecoveryResult = false;
 	const FTimespan MaximumWaitForHealth = FTimespan::FromSeconds(30);
@@ -1398,7 +1646,7 @@ FZenServiceInstance::TryRecovery()
 
 	if (TimespanSinceLastRecovery > MinimumDurationSinceLastRecovery)
 	{
-		FSystemWideCriticalSection RecoveryCriticalSection(TEXT("ZenServerRecovery"), MaximumWaitForHealth);
+		FSystemWideCriticalSection RecoveryCriticalSection(TEXT("ZenServerLaunch"), MaximumWaitForHealth);
 		if (!RecoveryCriticalSection.IsValid())
 		{
 			// A recovery is already in progress but did not complete in time, we assume we failed and let recovery continue on a different thread
@@ -1437,7 +1685,8 @@ FZenServiceInstance::TryRecovery()
 					return false;
 				}
 
-				AutoLaunch(Settings.SettingsVariant.Get<FServiceAutoLaunchSettings>(), *GetLocalServiceInstallPath(), HostName, Port);
+				AutoLaunch(Settings.SettingsVariant.Get<FServiceAutoLaunchSettings>(), ExecutablePath, ExecutionContextFilePath, HostName, Port);
+
 				FDateTime StartedWaitingForHealth = FDateTime::UtcNow();
 				bLastRecoveryResult = IsServiceReady();
 				while (!bLastRecoveryResult)
@@ -1452,7 +1701,7 @@ FZenServiceInstance::TryRecovery()
 					FPlatformProcess::Sleep(0.5f);
 					if (!IsZenProcessUsingEffectivePort(Port))
 					{
-						AutoLaunch(Settings.SettingsVariant.Get<FServiceAutoLaunchSettings>(), *GetLocalServiceInstallPath(), HostName, Port);
+						AutoLaunch(Settings.SettingsVariant.Get<FServiceAutoLaunchSettings>(), ExecutablePath, ExecutionContextFilePath, HostName, Port);
 					}
 					bLastRecoveryResult = IsServiceReady();
 				}
@@ -1504,9 +1753,14 @@ FZenServiceInstance::Initialize()
 {
 	if (Settings.IsAutoLaunch())
 	{
-		FString ExecutableInstallPath = ConditionalUpdateLocalInstall();
-		if (!ExecutableInstallPath.IsEmpty())
+		uint64 ZenAutoLaunchStartTime = FPlatformTime::Cycles64();
+		const FServiceAutoLaunchSettings& AutoLaunchSettings = Settings.SettingsVariant.Get<FServiceAutoLaunchSettings>();
+		bool ServiceIsInstalled = ConditionalUpdateLocalInstall(AutoLaunchSettings.InstallMode);
+		if (ServiceIsInstalled)
 		{
+			const FString ExecutablePath = GetLocalServiceInstallPath();
+			const FString ExecutionContextFilePath = GetServiceRunContextPath();
+
 			int LaunchAttempts = 0;
 			const FTimespan MaximumWaitForHealth = FTimespan::FromSeconds(20);
 
@@ -1514,15 +1768,14 @@ FZenServiceInstance::Initialize()
 			while (true)
 			{
 				{
-					FSystemWideCriticalSection RecoveryCriticalSection(TEXT("ZenServerRecovery"), MaximumWaitForHealth);
+					FSystemWideCriticalSection RecoveryCriticalSection(TEXT("ZenServerLaunch"), FTimespan::FromSeconds(5));
 					if (!RecoveryCriticalSection.IsValid())
 					{
 						UE_LOG(LogZenServiceInstance, Warning, TEXT("Local ZenServer AutoLaunch initialization timed out waiting for other process to complete startup/recovery"));
 					}
 					else
 					{
-						++LaunchAttempts;
-						bHasLaunchedLocal = AutoLaunch(Settings.SettingsVariant.Get<FServiceAutoLaunchSettings>(), *ExecutableInstallPath, HostName, Port);
+						bHasLaunchedLocal = AutoLaunch(AutoLaunchSettings, ExecutablePath, ExecutionContextFilePath, HostName, Port);
 						if (bHasLaunchedLocal)
 						{
 							const ZenServerState State(/*ReadOnly*/true);
@@ -1541,6 +1794,7 @@ FZenServiceInstance::Initialize()
 					}
 				}
 
+				++LaunchAttempts;
 				FTimespan WaitForHealth = FDateTime::UtcNow() - StartedWaitingForHealth;
 				if ((WaitForHealth > MaximumWaitForHealth) && (LaunchAttempts > 1))
 				{
@@ -1553,6 +1807,8 @@ FZenServiceInstance::Initialize()
 				FPlatformProcess::Sleep(0.5f);
 			}
 		}
+		double ZenAutoLaunchDuration = FPlatformTime::ToSeconds64(FPlatformTime::Cycles64() - ZenAutoLaunchStartTime);
+		UE_LOG(LogZenServiceInstance, Log, TEXT("Local ZenServer AutoLaunch initialization completed in %.3lf seconds"), ZenAutoLaunchDuration);
 	}
 	else
 	{
@@ -1618,87 +1874,215 @@ PromptUserOfFailedShutDownOfExistingProcess(uint16 Port)
 	}
 }
 
-FString
-FZenServiceInstance::ConditionalUpdateLocalInstall()
+bool
+FZenServiceInstance::ConditionalUpdateLocalInstall(FServiceAutoLaunchSettings::EInstallMode InstallMode)
 {
-	FString InTreeUtilityPath = FPaths::ConvertRelativePathToFull(FPlatformProcess::GenerateApplicationPath(TEXT("zen"), EBuildConfiguration::Development));
-	FString InstallUtilityPath = GetLocalInstallUtilityPath();
+	if (InstallMode == FServiceAutoLaunchSettings::EInstallMode::Link)
+	{
+		const FString LinkPath = GetServiceLinkPath();
+		IFileManager& FileManager = IFileManager::Get();
 
-	FString InTreeServicePath = FPaths::ConvertRelativePathToFull(FPlatformProcess::GenerateApplicationPath(TEXT("zenserver"), EBuildConfiguration::Development));
-	FString InstallServicePath = GetLocalServiceInstallPath();
+		bool LinkIsValid = false;
+		const FZenServiceLink Link = FZenServiceLink::Read(LinkPath);
+		if (Link)
+		{
+			if (!FileManager.FileExists(*Link.ServicePath))
+			{
+				UE_LOG(LogZenServiceInstance, Warning, TEXT("Found service link file '%s' pointing to unreachable service executable '%s'"), *LinkPath, *Link.ServicePath);
+			}
+			else if (!FileManager.FileExists(*Link.UtilityPath))
+			{
+				UE_LOG(LogZenServiceInstance, Warning, TEXT("Found service link file '%s' pointing to unreachable utility executable '%s'"), *LinkPath, *Link.UtilityPath);
+			}
+			else
+			{
+				LinkIsValid = true;
+			}
+		}
+		else if (FileManager.FileExists(*LinkPath))
+		{
+			UE_LOG(LogZenServiceInstance, Warning, TEXT("Found invalid service link file '%s', ignoring it"), *LinkPath);
+		}
+		
+		FString InTreeServicePath = GetInTreeServicePath();
+		if (LinkIsValid && InTreeServicePath == Link.ServicePath)
+		{
+			// If the running process already points to our executable and we have a valid link file we are good to go
+			uint32_t Pid = 0;
+			if (ZenServerState::FindRunningProcessId(*InTreeServicePath, &Pid))
+			{
+				UE_LOG(LogZenServiceInstance, Log, TEXT("Service link '%s' pointing to '%s', version %s is up and running"), *LinkPath, *Link.ServicePath, *Link.Version.ToString(false));
+				return true;
+			}
+		}
 
-	IFileManager& FileManager = IFileManager::Get();
+		FString InTreeUtilityPath = GetInTreeUtilityPath();
+		FZenVersion InTreeVersion;
+		if (!GetZenVersion(InTreeUtilityPath, InTreeServicePath, InTreeVersion))
+		{
+			checkf(false, TEXT("Unable to determine version using zen utility executable path: '%s'."), *InTreeUtilityPath);
+			return false;
+		}
+		if (LinkIsValid)
+		{
+			if (Link.Version < InTreeVersion)
+			{
+				UE_LOG(LogZenServiceInstance, Display, TEXT("Installing service link '%s' to '%s', version %s"), *LinkPath, *InTreeServicePath, *InTreeVersion.ToString(false));
+			}
+			else
+			{
+				// If the instance is running, assume it is valid and up to date
+				uint32_t Pid = 0;
+				if (ZenServerState::FindRunningProcessId(*Link.ServicePath, &Pid))
+				{
+					UE_LOG(LogZenServiceInstance, Display, TEXT("Service link '%s' pointing to '%s', version %s is up to date and running"), *LinkPath, *Link.ServicePath, *Link.Version.ToString(false));
+					return true;
+				}
 
-	bool bMainExecutablesUpdated = false;
-	FString InTreeVersionCache, InstallVersionCache;
-	if (IsInstallVersionOutOfDate(InTreeUtilityPath, InstallUtilityPath, InTreeServicePath, InstallServicePath, InTreeVersionCache, InstallVersionCache))
-	{
-		UE_LOG(LogZenServiceInstance, Display, TEXT("Installing service from '%s' to '%s'"), *InTreeServicePath, *InstallServicePath);
-		if (!ShutDownZenServerProcessExecutable(InstallServicePath))
+				// Verify that the executable pointed at is runnable and is of matching version
+				FZenVersion LinkedVersion;
+				if (GetZenVersion(Link.UtilityPath, Link.ServicePath, LinkedVersion))
+				{
+					if (LinkedVersion == Link.Version)
+					{
+						UE_LOG(LogZenServiceInstance, Display, TEXT("Service link '%s' pointing to '%s', version %s is up to date"), *LinkPath, *Link.ServicePath, *Link.Version.ToString(false));
+						return true;
+					}
+					else
+					{
+						UE_LOG(LogZenServiceInstance, Display, TEXT("Updating service link '%s' to '%s', version %s (link '%s', version %s does not match executable version %s)"), *LinkPath, *InTreeServicePath, *InTreeVersion.ToString(false), *Link.ServicePath, *Link.Version.ToString(false), *LinkedVersion.ToString(false));
+					}
+				}
+				else
+				{
+					UE_LOG(LogZenServiceInstance, Display, TEXT("Updating service link '%s' to '%s', version %s (link '%s', version %s pointing to invalid executable)"), *LinkPath, *InTreeServicePath, *InTreeVersion.ToString(false), *Link.ServicePath, *Link.Version.ToString(false));
+				}
+			}
+		}
+
+		FZenServiceLink NewLink{ .ServicePath = InTreeServicePath, .UtilityPath = InTreeUtilityPath, .Version = InTreeVersion };
+		if (!FZenServiceLink::Write(NewLink, LinkPath))
 		{
-			PromptUserToStopRunningServerInstanceForUpdate(InstallServicePath);
-			return FString();
+			UE_LOG(LogZenServiceInstance, Warning, TEXT("Failed to update service link file '%s'"), *LinkPath);
+			return false;
 		}
 
-		// Even after waiting for the process to shut down we have a tolerance for failure when overwriting the target files
-		if (!AttemptFileCopyWithRetries(*InstallServicePath, *InTreeServicePath, 5.0))
+		FString ServiceCopyInstallPath = GetServiceCopyInstallPath();
+		if (FileManager.FileExists(*ServiceCopyInstallPath))
 		{
-			PromptUserToStopRunningServerInstanceForUpdate(InstallServicePath);
-			return FString();
+			ShutDownZenServerProcessExecutable(ServiceCopyInstallPath);
 		}
 
-		if (!AttemptFileCopyWithRetries(*InstallUtilityPath, *InTreeUtilityPath, 5.0))
+		TArray<FString> FilesToCleanUp{ GetUtilityCopyInstallPath(), ServiceCopyInstallPath, GetInstallVersionCachePath(), GetInstallVersionCachePath() };
+#if PLATFORM_WINDOWS
+		FilesToCleanUp.Add(FPaths::SetExtension(FilesToCleanUp[0], TEXT("pdb")));
+		FilesToCleanUp.Add(FPaths::SetExtension(FilesToCleanUp[1], TEXT("pdb")));
+#endif // PLATFORM_WINDOWS
+		for (const FString& FileToCleanUp : FilesToCleanUp)
 		{
-			PromptUserToStopRunningServerInstanceForUpdate(InstallServicePath);
-			return FString();
+			// If zenserver is still running we may fail to clean up a file. Not critical, just try again next startup
+			if (!AttemptFileDeleteWithRetries(*FileToCleanUp, 1.0))
+			{
+				UE_LOG(LogZenServiceInstance, Log, TEXT("Failed cleaning up file {%s} (not critical)"), *FileToCleanUp);
+			}
 		}
+		return true;
+	}
+	else if (InstallMode == FServiceAutoLaunchSettings::EInstallMode::Copy)
+	{
+		FString InTreeUtilityPath = GetInTreeUtilityPath();
+		FString InstallUtilityPath = GetUtilityCopyInstallPath();
+		FString InTreeVersionCache = GetInTreeVersionCache();
 
-		AttemptFileCopyWithRetries(*InstallVersionCache, *InTreeVersionCache, 1.0);
+		FString InTreeServicePath = GetInTreeServicePath();
+		FString InstallServicePath = GetServiceCopyInstallPath();
+		FString InstallVersionCache = GetInstallVersionCachePath();
 
-		bMainExecutablesUpdated = true;
-	}
+		IFileManager& FileManager = IFileManager::Get();
 
-#if PLATFORM_WINDOWS
-	struct FZenExecutable
-	{
-		FString& InTreeFilePath;
-		FString& InstallFilePath;
-	};
-	const FZenExecutable ZenExecutables[] = {
-		// Service executable (zenserver.exe)
-		{InTreeServicePath, InstallServicePath},
-		// Utility executable (zen.exe)
-		{InTreeUtilityPath, InstallUtilityPath},
-	};
-	for (const FZenExecutable& Executable : ZenExecutables)
-	{
-		FString InTreeSymbolFilePath = FPaths::SetExtension(Executable.InTreeFilePath, TEXT("pdb"));
-		FString InstallSymbolFilePath = FPaths::SetExtension(Executable.InstallFilePath, TEXT("pdb"));
+		bool bMainExecutablesUpdated = false;
+		if (IsInstallVersionOutOfDate(InTreeUtilityPath, InstallUtilityPath, InTreeServicePath, InstallServicePath, InTreeVersionCache, InstallVersionCache))
+		{
+			UE_LOG(LogZenServiceInstance, Display, TEXT("Installing service from '%s' to '%s'"), *InTreeServicePath, *InstallServicePath);
+			if (!ShutDownZenServerProcessExecutable(InstallServicePath))
+			{
+				PromptUserToStopRunningServerInstanceForUpdate(InstallServicePath);
+				return false;
+			}
+
+			// Even after waiting for the process to shut down we have a tolerance for failure when overwriting the target files
+			if (!AttemptFileCopyWithRetries(*InstallServicePath, *InTreeServicePath, 5.0))
+			{
+				PromptUserToStopRunningServerInstanceForUpdate(InstallServicePath);
+				return false;
+			}
+
+			if (!AttemptFileCopyWithRetries(*InstallUtilityPath, *InTreeUtilityPath, 5.0))
+			{
+				PromptUserToStopRunningServerInstanceForUpdate(InstallServicePath);
+				return false;
+			}
+
+			AttemptFileCopyWithRetries(*InstallVersionCache, *InTreeVersionCache, 1.0);
+
+			bMainExecutablesUpdated = true;
+		}
 
-		if (FileManager.FileExists(*InTreeSymbolFilePath) && (bMainExecutablesUpdated || !FileManager.FileExists(*InstallSymbolFilePath)))
+#if PLATFORM_WINDOWS
+		struct FZenExecutable
+		{
+			FString& InTreeFilePath;
+			FString& InstallFilePath;
+		};
+		const FZenExecutable ZenExecutables[] = {
+			// Service executable (zenserver.exe)
+			{InTreeServicePath, InstallServicePath},
+			// Utility executable (zen.exe)
+			{InTreeUtilityPath, InstallUtilityPath},
+		};
+		for (const FZenExecutable& Executable : ZenExecutables)
 		{
-			AttemptFileCopyWithRetries(*InstallSymbolFilePath, *InTreeSymbolFilePath, 1.0);
+			FString InTreeSymbolFilePath = FPaths::SetExtension(Executable.InTreeFilePath, TEXT("pdb"));
+			FString InstallSymbolFilePath = FPaths::SetExtension(Executable.InstallFilePath, TEXT("pdb"));
+
+			if (FileManager.FileExists(*InTreeSymbolFilePath) && (bMainExecutablesUpdated || !FileManager.FileExists(*InstallSymbolFilePath)))
+			{
+				AttemptFileCopyWithRetries(*InstallSymbolFilePath, *InTreeSymbolFilePath, 1.0);
+			}
 		}
-	}
 #endif
 
-	FString InTreeCrashpadHandlerFilePath = FPaths::ConvertRelativePathToFull(FPlatformProcess::GenerateApplicationPath(TEXT("crashpad_handler"), EBuildConfiguration::Development));
-	FString InstallCrashpadHandlerFilePath = FPaths::ConvertRelativePathToFull(FPaths::Combine(GetLocalInstallPath(), FString(FPathViews::GetCleanFilename(InTreeCrashpadHandlerFilePath))));
+		FString InTreeCrashpadHandlerFilePath = GetInTreeCrashpadHandlerFilePath();
+		FString InstallCrashpadHandlerFilePath = GetInstallCrashpadHandlerFilePath(InTreeCrashpadHandlerFilePath);
+
+		if (FileManager.FileExists(*InTreeCrashpadHandlerFilePath) && (bMainExecutablesUpdated || !FileManager.FileExists(*InstallCrashpadHandlerFilePath)))
+		{
+			AttemptFileCopyWithRetries(*InstallCrashpadHandlerFilePath, *InTreeCrashpadHandlerFilePath, 1.0);
+		}
+
+		const FString LinkPath = GetServiceLinkPath();
+		if (FZenServiceLink Link = FZenServiceLink::Read(LinkPath); Link)
+		{
+			ShutDownZenServerProcessExecutable(Link.ServicePath);
+			if (!AttemptFileDeleteWithRetries(*LinkPath, 1.0))
+			{
+				UE_LOG(LogZenServiceInstance, Log, TEXT("Failed cleaning up file {%s} (not critical)"), *LinkPath);
+			}
+		}
 
-	if (FileManager.FileExists(*InTreeCrashpadHandlerFilePath) && (bMainExecutablesUpdated || !FileManager.FileExists(*InstallCrashpadHandlerFilePath)))
+		return true;
+	}
+	else
 	{
-		AttemptFileCopyWithRetries(*InstallCrashpadHandlerFilePath, *InTreeCrashpadHandlerFilePath, 1.0);
+		return false;
 	}
-
-	return InstallServicePath;
 }
 
 bool
-FZenServiceInstance::AutoLaunch(const FServiceAutoLaunchSettings& InSettings, FString&& ExecutablePath, FString& OutHostName, uint16& OutPort)
+FZenServiceInstance::AutoLaunch(const FServiceAutoLaunchSettings& InSettings, const FString& ExecutablePath, const FString& ExecutionContextFilePath, FString& OutHostName, uint16& OutPort)
 {
 	IFileManager& FileManager = IFileManager::Get();
 	const FString LockFilePath = FPaths::Combine(InSettings.DataPath, TEXT(".lock"));
-	const FString ExecutionContextFilePath = FPaths::SetExtension(ExecutablePath, TEXT(".runcontext"));
 
 	FString WorkingDirectory = FPaths::GetPath(ExecutablePath);
 
diff --git a/Engine/Source/Developer/Zen/Private/ZenServerState.cpp b/Engine/Source/Developer/Zen/Private/ZenServerState.cpp
--- a/Engine/Source/Developer/Zen/Private/ZenServerState.cpp
+++ b/Engine/Source/Developer/Zen/Private/ZenServerState.cpp
@@ -22,10 +22,74 @@
 #	include <sys/file.h>
 #	include <sys/mman.h>
 #	include <sys/sem.h>
+#	include <sys/stat.h>
+#	include <sys/sysctl.h>
 #endif
 
 DEFINE_LOG_CATEGORY_STATIC(LogZenServiceState, Log, All);
 
+#if PLATFORM_UNIX
+static char
+GetPidStatus(int Pid)
+{
+	TAnsiStringBuilder<128> StatPath;
+	StatPath.Appendf("/proc/%d/stat", Pid);
+	FILE* StatFile = fopen(*StatPath, "r");
+	if (StatFile)
+	{
+		char Buffer[5120];
+		int	 Size = fread(Buffer, 1, 5120 - 1, StatFile);
+		fclose(StatFile);
+		if (Size > 0)
+		{
+			Buffer[Size + 1] = 0;
+			char* ScanPtr = strrchr(Buffer, ')');
+			if (ScanPtr && ScanPtr[1] != '\0')
+			{
+				ScanPtr += 2;
+				char State = *ScanPtr;
+				return State;
+			}
+		}
+	}
+	return 0;
+}
+
+static bool
+IsZombieProcess(int pid)
+{
+	char Status = GetPidStatus(pid);
+	if (Status == 'Z' || Status == 0)
+	{
+		return true;
+	}
+	return false;
+}
+
+#endif	// ZEN_PLATFORM_LINUX
+
+#if PLATFORM_MAC
+static bool
+IsZombieProcess(int pid)
+{
+	struct kinfo_proc Info;
+	int				  Mib[4] = { CTL_KERN, KERN_PROC, KERN_PROC_PID, pid };
+	size_t			  InfoSize = sizeof Info;
+
+	int Res = sysctl(Mib, 4, &Info, &InfoSize, NULL, 0);
+	if (Res != 0)
+	{
+		return false;
+	}
+	if (Info.kp_proc.p_stat == SZOMB)
+	{
+		// Zombie process
+		return true;
+	}
+	return false;
+}
+#endif	// PLATFORM_MAC
+
 // Native functions to interact with a process using a process id
 // We don't use UE's own OpenProcess as they try to open processes with PROCESS_ALL_ACCESS
 bool ZenServerState::IsProcessRunning(uint32 Pid)
@@ -73,6 +137,10 @@ bool ZenServerState::IsProcessRunning(uint32 Pid)
 	int Res = kill(pid_t(Pid), 0);
 	if (Res == 0)
 	{
+		if (IsZombieProcess(Pid))
+		{
+			return false;
+		}
 		return true;
 	}
 	int Error = errno;
diff --git a/Engine/Source/Developer/Zen/Private/ZenVersion.cpp b/Engine/Source/Developer/Zen/Private/ZenVersion.cpp
--- a/Engine/Source/Developer/Zen/Private/ZenVersion.cpp
+++ b/Engine/Source/Developer/Zen/Private/ZenVersion.cpp
@@ -111,4 +111,12 @@ bool FZenVersion::operator<(FZenVersion& Other) const
 	return false;
 }
 
+FZenVersion::operator bool() const
+{
+	return 	MajorVersion != 0 ||
+		MinorVersion != 0 ||
+		PatchVersion != 0 ||
+		!Details.IsEmpty();
+}
+
 } // namespace UE::Zen
diff --git a/Engine/Source/Developer/Zen/Public/Experimental/ZenServerInterface.h b/Engine/Source/Developer/Zen/Public/Experimental/ZenServerInterface.h
--- a/Engine/Source/Developer/Zen/Public/Experimental/ZenServerInterface.h
+++ b/Engine/Source/Developer/Zen/Public/Experimental/ZenServerInterface.h
@@ -47,6 +47,11 @@ struct FServiceAutoLaunchSettings
 	bool bAllowRemoteNetworkService = false;
 	bool bSendUnattendedBugReports = false;
 	bool bIsDefaultSharedRunContext = true;
+	enum class EInstallMode
+	{
+		Copy,
+		Link
+	} InstallMode = EInstallMode::Copy;
 };
 
 struct FServiceSettings
@@ -205,8 +210,8 @@ public:
 private:
 
 	void Initialize();
-	FString ConditionalUpdateLocalInstall();
-	static bool AutoLaunch(const FServiceAutoLaunchSettings& InSettings, FString&& ExecutablePath, FString& OutHostName, uint16& OutPort);
+	bool ConditionalUpdateLocalInstall(FServiceAutoLaunchSettings::EInstallMode InstallMode);
+	static bool AutoLaunch(const FServiceAutoLaunchSettings& InSettings, const FString& ExecutablePath, const FString& ExecutionContextFilePath, FString& OutHostName, uint16& OutPort);
 	
 	mutable TPimplPtr<class FZenHttpRequest> GCStatusHttpRequest;
 	mutable TFuture<TOptional<FGCStatus>> GCStatusRequest;
diff --git a/Engine/Source/Developer/Zen/Tests/ZenTest.cpp b/Engine/Source/Developer/Zen/Tests/ZenTest.cpp
--- a/Engine/Source/Developer/Zen/Tests/ZenTest.cpp
+++ b/Engine/Source/Developer/Zen/Tests/ZenTest.cpp
@@ -23,15 +23,16 @@ TEST_CASE("Zen::ZenServerInterface", "[Zen][Basic]")
 	FString DefaultDataPath = FPaths::ConvertRelativePathToFull(FPaths::Combine(DataPathRoot, "Default"));
 	IFileManager::Get().DeleteDirectory(*DataPathRoot, false, true);
 
-	SECTION("Basic AutoLaunch and Shutdown")
+	SECTION("Basic AutoLaunch and Shutdown (copy)")
 	{
-		for (int Iteration = 0; Iteration < 2; ++Iteration)
+		for (int Iteration = 0; Iteration < 4; ++Iteration)
 		{
 			FServiceSettings ZenTestServiceSettings;
 			FServiceAutoLaunchSettings& ZenTestAutoLaunchSettings = ZenTestServiceSettings.SettingsVariant.Get<FServiceAutoLaunchSettings>();
 			ZenTestAutoLaunchSettings.DataPath = DefaultDataPath;
 			ZenTestAutoLaunchSettings.ExtraArgs = DefaultArgs;
 			ZenTestAutoLaunchSettings.DesiredPort = DefaultTestPort;
+			ZenTestAutoLaunchSettings.InstallMode = FServiceAutoLaunchSettings::EInstallMode::Copy;
 
 			{
 				FScopeZenService ScopeZenService(MoveTemp(ZenTestServiceSettings));
@@ -48,15 +49,68 @@ TEST_CASE("Zen::ZenServerInterface", "[Zen][Basic]")
 		}
 	}
 
-	SECTION("Overlapping AutoLaunch and Shutdown")
+	SECTION("Basic AutoLaunch and Shutdown (link)")
 	{
-		for (int Iteration = 0; Iteration < 3; ++Iteration)
+		for (int Iteration = 0; Iteration < 4; ++Iteration)
 		{
 			FServiceSettings ZenTestServiceSettings;
 			FServiceAutoLaunchSettings& ZenTestAutoLaunchSettings = ZenTestServiceSettings.SettingsVariant.Get<FServiceAutoLaunchSettings>();
 			ZenTestAutoLaunchSettings.DataPath = DefaultDataPath;
 			ZenTestAutoLaunchSettings.ExtraArgs = DefaultArgs;
 			ZenTestAutoLaunchSettings.DesiredPort = DefaultTestPort;
+			ZenTestAutoLaunchSettings.InstallMode = FServiceAutoLaunchSettings::EInstallMode::Link;
+
+			{
+				FScopeZenService ScopeZenService(MoveTemp(ZenTestServiceSettings));
+				FZenServiceInstance& ZenInstance = ScopeZenService.GetInstance();
+				uint16 AutoLaunchedPort = ZenInstance.GetAutoLaunchedPort();
+				uint16 DetectedPort = 0;
+				CHECK(ZenInstance.IsServiceReady());
+
+				CHECK(IsLocalServiceRunning(*DefaultDataPath, &DetectedPort));
+				CHECK(DetectedPort == AutoLaunchedPort);
+				CHECK(StopLocalService(*DefaultDataPath));
+				CHECK(!IsLocalServiceRunning(*DefaultDataPath));
+			}
+		}
+	}
+
+	SECTION("Overlapping AutoLaunch and Shutdown (copy)")
+	{
+		for (int Iteration = 0; Iteration < 11; ++Iteration)
+		{
+			FServiceSettings ZenTestServiceSettings;
+			FServiceAutoLaunchSettings& ZenTestAutoLaunchSettings = ZenTestServiceSettings.SettingsVariant.Get<FServiceAutoLaunchSettings>();
+			ZenTestAutoLaunchSettings.DataPath = DefaultDataPath;
+			ZenTestAutoLaunchSettings.ExtraArgs = DefaultArgs;
+			ZenTestAutoLaunchSettings.DesiredPort = DefaultTestPort;
+			ZenTestAutoLaunchSettings.InstallMode = FServiceAutoLaunchSettings::EInstallMode::Copy;
+
+			{
+				FScopeZenService ScopeZenService(MoveTemp(ZenTestServiceSettings));
+				FZenServiceInstance& ZenInstance = ScopeZenService.GetInstance();
+				uint16 AutoLaunchedPort = ZenInstance.GetAutoLaunchedPort();
+				uint16 DetectedPort = 0;
+				CHECK(ZenInstance.IsServiceReady());
+
+				CHECK(IsLocalServiceRunning(*DefaultDataPath, &DetectedPort));
+				CHECK(DetectedPort == AutoLaunchedPort);
+			}
+		}
+		CHECK(StopLocalService(*DefaultDataPath));
+		CHECK(!IsLocalServiceRunning(*DefaultDataPath));
+	}
+
+	SECTION("Overlapping AutoLaunch and Shutdown (link)")
+	{
+		for (int Iteration = 0; Iteration < 11; ++Iteration)
+		{
+			FServiceSettings ZenTestServiceSettings;
+			FServiceAutoLaunchSettings& ZenTestAutoLaunchSettings = ZenTestServiceSettings.SettingsVariant.Get<FServiceAutoLaunchSettings>();
+			ZenTestAutoLaunchSettings.DataPath = DefaultDataPath;
+			ZenTestAutoLaunchSettings.ExtraArgs = DefaultArgs;
+			ZenTestAutoLaunchSettings.DesiredPort = DefaultTestPort;
+			ZenTestAutoLaunchSettings.InstallMode = FServiceAutoLaunchSettings::EInstallMode::Link;
 
 			{
 				FScopeZenService ScopeZenService(MoveTemp(ZenTestServiceSettings));
@@ -73,15 +127,68 @@ TEST_CASE("Zen::ZenServerInterface", "[Zen][Basic]")
 		CHECK(!IsLocalServiceRunning(*DefaultDataPath));
 	}
 
-	SECTION("Overlapping AutoLaunch and Shutdown With DataPath Shared And Differing Args")
+	SECTION("Overlapping AutoLaunch and Shutdown With DataPath Shared And Differing Args (copy)")
 	{
-		for (int Iteration = 0; Iteration < 3; ++Iteration)
+		for (int Iteration = 0; Iteration < 11; ++Iteration)
 		{
 			FServiceSettings ZenTestServiceSettings;
 			FServiceAutoLaunchSettings& ZenTestAutoLaunchSettings = ZenTestServiceSettings.SettingsVariant.Get<FServiceAutoLaunchSettings>();
 			ZenTestAutoLaunchSettings.DataPath = DefaultDataPath;
 			ZenTestAutoLaunchSettings.ExtraArgs = *WriteToString<128>(DefaultArgs, TEXT(" --gc-interval-seconds "), (Iteration+1)*1000);
 			ZenTestAutoLaunchSettings.DesiredPort = DefaultTestPort;
+			ZenTestAutoLaunchSettings.InstallMode = FServiceAutoLaunchSettings::EInstallMode::Copy;
+
+			{
+				FScopeZenService ScopeZenService(MoveTemp(ZenTestServiceSettings));
+				FZenServiceInstance& ZenInstance = ScopeZenService.GetInstance();
+				uint16 AutoLaunchedPort = ZenInstance.GetAutoLaunchedPort();
+				uint16 DetectedPort = 0;
+				CHECK(ZenInstance.IsServiceReady());
+
+				CHECK(IsLocalServiceRunning(*DefaultDataPath, &DetectedPort));
+				CHECK(DetectedPort == AutoLaunchedPort);
+			}
+		}
+		CHECK(StopLocalService(*DefaultDataPath));
+		CHECK(!IsLocalServiceRunning(*DefaultDataPath));
+	}
+
+	SECTION("Overlapping AutoLaunch and Shutdown With DataPath Shared And Differing Args (link)")
+	{
+		for (int Iteration = 0; Iteration < 11; ++Iteration)
+		{
+			FServiceSettings ZenTestServiceSettings;
+			FServiceAutoLaunchSettings& ZenTestAutoLaunchSettings = ZenTestServiceSettings.SettingsVariant.Get<FServiceAutoLaunchSettings>();
+			ZenTestAutoLaunchSettings.DataPath = DefaultDataPath;
+			ZenTestAutoLaunchSettings.ExtraArgs = *WriteToString<128>(DefaultArgs, TEXT(" --gc-interval-seconds "), (Iteration + 1) * 1000);
+			ZenTestAutoLaunchSettings.DesiredPort = DefaultTestPort;
+			ZenTestAutoLaunchSettings.InstallMode = FServiceAutoLaunchSettings::EInstallMode::Link;
+
+			{
+				FScopeZenService ScopeZenService(MoveTemp(ZenTestServiceSettings));
+				FZenServiceInstance& ZenInstance = ScopeZenService.GetInstance();
+				uint16 AutoLaunchedPort = ZenInstance.GetAutoLaunchedPort();
+				uint16 DetectedPort = 0;
+				CHECK(ZenInstance.IsServiceReady());
+
+				CHECK(IsLocalServiceRunning(*DefaultDataPath, &DetectedPort));
+				CHECK(DetectedPort == AutoLaunchedPort);
+			}
+		}
+		CHECK(StopLocalService(*DefaultDataPath));
+		CHECK(!IsLocalServiceRunning(*DefaultDataPath));
+	}
+
+	SECTION("Overlapping AutoLaunch and Shutdown With DataPath Shared And Differing Args (alternating)")
+	{
+		for (int Iteration = 0; Iteration < 11; ++Iteration)
+		{
+			FServiceSettings ZenTestServiceSettings;
+			FServiceAutoLaunchSettings& ZenTestAutoLaunchSettings = ZenTestServiceSettings.SettingsVariant.Get<FServiceAutoLaunchSettings>();
+			ZenTestAutoLaunchSettings.DataPath = DefaultDataPath;
+			ZenTestAutoLaunchSettings.ExtraArgs = *WriteToString<128>(DefaultArgs, TEXT(" --gc-interval-seconds "), (Iteration + 1) * 1000);
+			ZenTestAutoLaunchSettings.DesiredPort = DefaultTestPort;
+			ZenTestAutoLaunchSettings.InstallMode = ((Iteration & 1) == 0) ? FServiceAutoLaunchSettings::EInstallMode::Copy : FServiceAutoLaunchSettings::EInstallMode::Link;
 
 			{
 				FScopeZenService ScopeZenService(MoveTemp(ZenTestServiceSettings));
diff --git a/Engine/Source/Runtime/Core/Private/Mac/MacCriticalSection.cpp b/Engine/Source/Runtime/Core/Private/Mac/MacCriticalSection.cpp
--- a/Engine/Source/Runtime/Core/Private/Mac/MacCriticalSection.cpp
+++ b/Engine/Source/Runtime/Core/Private/Mac/MacCriticalSection.cpp
@@ -14,33 +14,44 @@
 
 FMacSystemWideCriticalSection::FMacSystemWideCriticalSection(const FString& InName, FTimespan InTimeout)
 {
-	check(InName.Len() > 0)
-	check(InTimeout >= FTimespan::Zero())
-	check(InTimeout.GetTotalSeconds() < (double)FLT_MAX)	// we'll need this to fit in a single-precision float later so check it here
+	check(InName.Len() > 0);
+	check(InTimeout >= FTimespan::Zero());
+	check(InTimeout.GetTotalSeconds() < (double)FLT_MAX);
 
-	FDateTime ExpireTime = FDateTime::UtcNow() + InTimeout;
-
-	// This lock implementation is using files so correct any backslashes in the name
-	const FString LockPath = FString(FMacPlatformProcess::ApplicationSettingsDir()) / InName;
+	const FString LockPath = FString(FPlatformProcess::ApplicationSettingsDir()) / InName;
 	FString NormalizedFilepath(LockPath);
 	NormalizedFilepath.ReplaceInline(TEXT("\\"), TEXT("/"));
+	FileHandle = -1;
 
-	// Attempt to open a file with O_EXLOCK (equivalent of atomic open() + flock())
-	FileHandle = open(TCHAR_TO_UTF8(*NormalizedFilepath), O_CREAT | O_WRONLY | O_EXLOCK | O_NONBLOCK, S_IRUSR | S_IWUSR | S_IRGRP | S_IWGRP | S_IROTH | S_IWOTH);
-
-	if (FileHandle == -1 && InTimeout != FTimespan::Zero())
+	double ExpireTimeSecs = FPlatformTime::Seconds() + InTimeout.GetTotalSeconds();
+	while (true)
 	{
-		// Failed to get a valid file handle so the file is probably open and locked by another owner - retry until we timeout
+		if (FileHandle == -1)
+		{
+			// Try to open the file.
+			FileHandle = open(TCHAR_TO_UTF8(*NormalizedFilepath), O_CREAT | O_WRONLY | O_NONBLOCK, S_IRUSR | S_IWUSR | S_IRGRP | S_IWGRP | S_IROTH | S_IWOTH);
+		}
 
-		for (FTimespan TimeoutRemaining = ExpireTime - FDateTime::UtcNow();
-			 FileHandle == -1 && TimeoutRemaining > FTimespan::Zero();
-			 TimeoutRemaining = ExpireTime - FDateTime::UtcNow())
+		// If the file is open, try to lock it, but don't block. If the file is already locked by another process or another thread, blocking here may not honor InTimeout.
+		if (FileHandle != -1 && flock(FileHandle, LOCK_EX | LOCK_NB) == 0)
 		{
-			// retry until timeout
-			float RetrySeconds = FMath::Min((float)TimeoutRemaining.GetTotalSeconds(), 0.25f);
-			FMacPlatformProcess::Sleep(RetrySeconds);
-			FileHandle = open(TCHAR_TO_UTF8(*NormalizedFilepath), O_CREAT | O_WRONLY | O_EXLOCK | O_NONBLOCK, S_IRUSR | S_IWUSR | S_IRGRP | S_IWGRP | S_IROTH | S_IWOTH);
+			return; // Lock was successfully taken.
 		}
+
+		// If the lock isn't acquired and no time is left to retry, clean up and set the state as 'invalid'
+		if (InTimeout == FTimespan::Zero() || FPlatformTime::Seconds() > ExpireTimeSecs)
+		{
+			if (FileHandle != -1)
+			{
+				close(FileHandle);
+				FileHandle = -1;
+			}
+			return; // Lock wasn't acquired within the allowed time.
+		}
+
+		// Either the file did not open or the lock wasn't acquired, retry.
+		const float RetrySeconds = FMath::Min((float)InTimeout.GetTotalSeconds(), 0.25f);
+		FPlatformProcess::Sleep(RetrySeconds);
 	}
 }
 
@@ -58,6 +69,7 @@ void FMacSystemWideCriticalSection::Release()
 {
 	if (IsValid())
 	{
+		flock(FileHandle, LOCK_UN);
 		close(FileHandle);
 		FileHandle = -1;
 	}
diff --git a/Engine/Source/Runtime/Core/Private/Unix/UnixCriticalSection.cpp b/Engine/Source/Runtime/Core/Private/Unix/UnixCriticalSection.cpp
--- a/Engine/Source/Runtime/Core/Private/Unix/UnixCriticalSection.cpp
+++ b/Engine/Source/Runtime/Core/Private/Unix/UnixCriticalSection.cpp
@@ -16,9 +16,9 @@
 
 FUnixSystemWideCriticalSection::FUnixSystemWideCriticalSection(const FString& InName, FTimespan InTimeout)
 {
-	check(InName.Len() > 0)
-	check(InTimeout >= FTimespan::Zero())
-	check(InTimeout.GetTotalSeconds() < (double)FLT_MAX)
+	check(InName.Len() > 0);
+	check(InTimeout >= FTimespan::Zero());
+	check(InTimeout.GetTotalSeconds() < (double)FLT_MAX);
 
 	const FString LockPath = FString(FPlatformProcess::ApplicationSettingsDir()) / InName;
 	FString NormalizedFilepath(LockPath);
