Remove caching of ZenHttpRequest instances in FZenServiceInstance as they cause lingering sockets when shutting down ZenServer via Zen::StopLocalService.

diff --git a/Engine/Source/Developer/Zen/Private/ZenServerInterface.cpp b/Engine/Source/Developer/Zen/Private/ZenServerInterface.cpp
--- a/Engine/Source/Developer/Zen/Private/ZenServerInterface.cpp
+++ b/Engine/Source/Developer/Zen/Private/ZenServerInterface.cpp
@@ -2385,7 +2385,6 @@ FZenServiceInstance::AutoLaunch(const FServiceAutoLaunchSettings& InSettings, co
 bool 
 FZenServiceInstance::GetCacheStats(FZenCacheStats& Stats)
 {
-	UE::Zen::FZenHttpRequest* Request = nullptr;
 	{
 		TUniqueLock Lock(LastCacheStatsMutex);
 		// If we've already requested stats and they are ready then grab them
@@ -2408,41 +2407,34 @@ FZenServiceInstance::GetCacheStats(FZenCacheStats& Stats)
 		{
 			return Stats.bIsValid;
 		}
-
-		if (!CacheStatsHttpRequest.IsValid())
-		{
-			TStringBuilder<128> ZenDomain;
-			ZenDomain << HostName << TEXT(":") << Port;
-			CacheStatsHttpRequest = MakePimpl<FZenHttpRequest>(ZenDomain.ToString(), false);
-		}
-
-		Request = CacheStatsHttpRequest.Get();
 	}
 
+	TStringBuilder<128> ZenDomain;
+	ZenDomain << HostName << TEXT(":") << Port;
+
 #if WITH_EDITOR
 	EAsyncExecution ThreadPool = EAsyncExecution::LargeThreadPool;
 #else
 	EAsyncExecution ThreadPool = EAsyncExecution::ThreadPool;
 #endif
 		// We've not got any requests in flight and we've met a given time requirement for requests
-	CacheStatsRequest = Async(ThreadPool, [Request]
-	{
-		check(Request != nullptr);
-		Request->Reset();
+	CacheStatsRequest = Async(ThreadPool, [ZenDomain = ZenDomain.ToString()]
+		{
+			UE::Zen::FZenHttpRequest Request(ZenDomain, false);
 
-		TArray64<uint8> GetBuffer;
-		FZenHttpRequest::Result Result = Request->PerformBlockingDownload(TEXTVIEW("/stats/z$"), &GetBuffer, Zen::EContentType::CbObject);
+			TArray64<uint8> GetBuffer;
+			FZenHttpRequest::Result Result = Request.PerformBlockingDownload(TEXTVIEW("/stats/z$"), &GetBuffer, Zen::EContentType::CbObject);
 
-		FZenCacheStats Stats;
+			FZenCacheStats Stats;
 
-		if (Result == Zen::FZenHttpRequest::Result::Success && Request->GetResponseCode() == 200)
-		{
-			FCbFieldView RootView(GetBuffer.GetData());
-			Stats.bIsValid = LoadFromCompactBinary(RootView, Stats);
-		}
+			if (Result == Zen::FZenHttpRequest::Result::Success && Request.GetResponseCode() == 200)
+			{
+				FCbFieldView RootView(GetBuffer.GetData());
+				Stats.bIsValid = LoadFromCompactBinary(RootView, Stats);
+			}
 
-		return Stats;
-	});
+			return Stats;
+		});
 
 	return Stats.bIsValid;
 }
@@ -2450,7 +2442,6 @@ FZenServiceInstance::GetCacheStats(FZenCacheStats& Stats)
 bool 
 FZenServiceInstance::GetProjectStats(FZenProjectStats& Stats)
 {
-	UE::Zen::FZenHttpRequest* Request = nullptr;
 	{
 		TUniqueLock Lock(LastProjectStatsMutex);
 		// If we've already requested stats and they are ready then grab them
@@ -2473,32 +2464,27 @@ FZenServiceInstance::GetProjectStats(FZenProjectStats& Stats)
 		{
 			return Stats.bIsValid;
 		}
-		if (!ProjectStatsHttpRequest.IsValid())
-		{
-			TStringBuilder<128> ZenDomain;
-			ZenDomain << HostName << TEXT(":") << Port;
-			ProjectStatsHttpRequest = MakePimpl<FZenHttpRequest>(ZenDomain.ToString(), false);
-		}
-		Request = ProjectStatsHttpRequest.Get();
 	}
 
+	TStringBuilder<128> ZenDomain;
+	ZenDomain << HostName << TEXT(":") << Port;
+
 #if WITH_EDITOR
 	EAsyncExecution ThreadPool = EAsyncExecution::LargeThreadPool;
 #else
 	EAsyncExecution ThreadPool = EAsyncExecution::ThreadPool;
 #endif
 			// We've not got any requests in flight and we've met a given time requirement for requests
-	ProjectStatsRequest = Async(ThreadPool, [Request]
-	{
-		check(Request);
-		Request->Reset();
+	ProjectStatsRequest = Async(ThreadPool, [ZenDomain = ZenDomain.ToString()]
+		{
+			UE::Zen::FZenHttpRequest Request(ZenDomain, false);
 
 			TArray64<uint8> GetBuffer;
-		FZenHttpRequest::Result Result = Request->PerformBlockingDownload(TEXTVIEW("/stats/prj"), &GetBuffer, Zen::EContentType::CbObject);
+			FZenHttpRequest::Result Result = Request.PerformBlockingDownload(TEXTVIEW("/stats/prj"), &GetBuffer, Zen::EContentType::CbObject);
 
 			FZenProjectStats Stats;
 
-		if (Result == Zen::FZenHttpRequest::Result::Success && Request->GetResponseCode() == 200)
+			if (Result == Zen::FZenHttpRequest::Result::Success && Request.GetResponseCode() == 200)
 			{
 				FCbFieldView RootView(GetBuffer.GetData());
 				Stats.bIsValid = LoadFromCompactBinary(RootView, Stats);
@@ -2541,18 +2527,13 @@ FZenServiceInstance::GetGCStatus(FGCStatus& Status)
 #else
 		EAsyncExecution ThreadPool = EAsyncExecution::ThreadPool;
 #endif
-		if (!GCStatusHttpRequest.IsValid())
-		{
-			TStringBuilder<128> ZenDomain;
-			ZenDomain << HostName << TEXT(":") << Port;
-			GCStatusHttpRequest = MakePimpl<FZenHttpRequest>(ZenDomain.ToString(), false);
-		}
+		TStringBuilder<128> ZenDomain;
+		ZenDomain << HostName << TEXT(":") << Port;
 
 		// We've not got any requests in flight and we've met a given time requirement for requests
-		GCStatusRequest = Async(ThreadPool, [this]
+		GCStatusRequest = Async(ThreadPool, [this, ZenDomain = ZenDomain.ToString()]
 			{
-				UE::Zen::FZenHttpRequest& Request = *GCStatusHttpRequest.Get();
-				Request.Reset();
+				UE::Zen::FZenHttpRequest Request(ZenDomain, false);
 
 				TArray64<uint8> GetBuffer;
 				FZenHttpRequest::Result Result = Request.PerformBlockingDownload(TEXTVIEW("/admin/gc"), &GetBuffer, Zen::EContentType::CbObject);
diff --git a/Engine/Source/Developer/Zen/Public/Experimental/ZenServerInterface.h b/Engine/Source/Developer/Zen/Public/Experimental/ZenServerInterface.h
--- a/Engine/Source/Developer/Zen/Public/Experimental/ZenServerInterface.h
+++ b/Engine/Source/Developer/Zen/Public/Experimental/ZenServerInterface.h
@@ -213,18 +213,15 @@ private:
 	bool ConditionalUpdateLocalInstall(FServiceAutoLaunchSettings::EInstallMode InstallMode);
 	static bool AutoLaunch(const FServiceAutoLaunchSettings& InSettings, const FString& ExecutablePath, const FString& ExecutionContextFilePath, FString& OutHostName, uint16& OutPort);
 	
-	mutable TPimplPtr<class FZenHttpRequest> GCStatusHttpRequest;
 	mutable TFuture<TOptional<FGCStatus>> GCStatusRequest;
 	mutable TOptional<FGCStatus> LastGCStatus;
 	mutable uint64 LastGCStatusTime = 0;
 
-	mutable TPimplPtr<class FZenHttpRequest> CacheStatsHttpRequest;
 	mutable TFuture<FZenCacheStats> CacheStatsRequest;
 	mutable uint64 LastCacheStatsTime = 0;
 	mutable FMutex LastCacheStatsMutex;
 	mutable FZenCacheStats LastCacheStats;
 
-	mutable TPimplPtr<class FZenHttpRequest> ProjectStatsHttpRequest;
 	mutable TFuture<FZenProjectStats> ProjectStatsRequest;
 	mutable uint64 LastProjectStatsTime = 0;
 	mutable FMutex LastProjectStatsMutex;
