Fixing a TCHAR* into a stack variable being passed as a capture to an async lambda method by accident.

diff --git a/Engine/Source/Developer/Zen/Private/ZenServerInterface.cpp b/Engine/Source/Developer/Zen/Private/ZenServerInterface.cpp
--- a/Engine/Source/Developer/Zen/Private/ZenServerInterface.cpp
+++ b/Engine/Source/Developer/Zen/Private/ZenServerInterface.cpp
@@ -2418,7 +2418,7 @@ FZenServiceInstance::GetCacheStats(FZenCacheStats& Stats)
 	EAsyncExecution ThreadPool = EAsyncExecution::ThreadPool;
 #endif
 		// We've not got any requests in flight and we've met a given time requirement for requests
-	CacheStatsRequest = Async(ThreadPool, [ZenDomain = ZenDomain.ToString()]
+	CacheStatsRequest = Async(ThreadPool, [ZenDomain = FString(ZenDomain)]
 		{
 			UE::Zen::FZenHttpRequest Request(ZenDomain, false);
 
@@ -2475,7 +2475,7 @@ FZenServiceInstance::GetProjectStats(FZenProjectStats& Stats)
 	EAsyncExecution ThreadPool = EAsyncExecution::ThreadPool;
 #endif
 			// We've not got any requests in flight and we've met a given time requirement for requests
-	ProjectStatsRequest = Async(ThreadPool, [ZenDomain = ZenDomain.ToString()]
+	ProjectStatsRequest = Async(ThreadPool, [ZenDomain = FString(ZenDomain)]
 		{
 			UE::Zen::FZenHttpRequest Request(ZenDomain, false);
 
@@ -2531,7 +2531,7 @@ FZenServiceInstance::GetGCStatus(FGCStatus& Status)
 		ZenDomain << HostName << TEXT(":") << Port;
 
 		// We've not got any requests in flight and we've met a given time requirement for requests
-		GCStatusRequest = Async(ThreadPool, [this, ZenDomain = ZenDomain.ToString()]
+		GCStatusRequest = Async(ThreadPool, [this, ZenDomain = FString(ZenDomain)]
 			{
 				UE::Zen::FZenHttpRequest Request(ZenDomain, false);
 
