Make sure we retain limitedprogresslifetime setting when launching zenserver from Zen Dashboard
Decrease default connection timeout in FZenHttpRequest to 5 sec (from 30 sec)
Make sure we verify that the running zenserver instance picks up our process if we request to attach to it
Only use Zen::FServiceSettings from manifest file if it actually contains a "zenserver"/"settings" section
Remove obsolete `--PackageStoreManifest` option for UnrealPack `StartZenServerForStage` command
Shorten timeout when checking if zenserver is ready during startup and recovery

#rnx

diff --git a/Engine/Source/Developer/IoStoreUtilities/Private/CookedPackageStore.cpp b/Engine/Source/Developer/IoStoreUtilities/Private/CookedPackageStore.cpp
--- a/Engine/Source/Developer/IoStoreUtilities/Private/CookedPackageStore.cpp
+++ b/Engine/Source/Developer/IoStoreUtilities/Private/CookedPackageStore.cpp
@@ -35,12 +35,18 @@ FIoStatus FCookedPackageStore::LoadManifest(const TCHAR* ManifestFilename)
 	FCbObject OplogObject;
 	if (FCbFieldView ZenServerField = ManifestObject["zenserver"])
 	{
-		UE::Zen::FServiceSettings ZenServiceSettings;
-		ZenServiceSettings.ReadFromCompactBinary(ZenServerField["settings"]);
 		FString ProjectId = FString(ZenServerField["projectid"].AsString());
 		FString OplogId = FString(ZenServerField["oplogid"].AsString());
 
-		ZenStoreClient = MakeUnique<UE::FZenStoreHttpClient>(MoveTemp(ZenServiceSettings));
+		UE::Zen::FServiceSettings ZenServiceSettings;
+		if (ZenServiceSettings.ReadFromCompactBinary(ZenServerField["settings"]))
+		{
+			ZenStoreClient = MakeUnique<UE::FZenStoreHttpClient>(MoveTemp(ZenServiceSettings));
+		}
+		else
+		{
+			ZenStoreClient = MakeUnique<UE::FZenStoreHttpClient>();
+		}
 		ZenStoreClient->InitializeReadOnly(ProjectId, OplogId);
 
 		COOKEDPACKAGESTORE_CPU_SCOPE(FetchOplog);
diff --git a/Engine/Source/Developer/IoStoreUtilities/Private/IoStoreUtilities.cpp b/Engine/Source/Developer/IoStoreUtilities/Private/IoStoreUtilities.cpp
--- a/Engine/Source/Developer/IoStoreUtilities/Private/IoStoreUtilities.cpp
+++ b/Engine/Source/Developer/IoStoreUtilities/Private/IoStoreUtilities.cpp
@@ -9605,41 +9605,6 @@ int32 CreateIoStoreContainerFiles(const TCHAR* CmdLine)
 			SponsorProcessIds.Add(SponsorProcessId);
 		}
 
-		FString ManifestFilename;
-		if (FParse::Value(FCommandLine::Get(), TEXT("PackageStoreManifest="), ManifestFilename))
-		{
-			TUniquePtr<FArchive> Ar(IFileManager::Get().CreateFileReader(*ManifestFilename));
-			if (Ar)
-			{
-				FCbObject ManifestObject = LoadCompactBinary(*Ar).AsObject();
-				FCbObject OplogObject;
-				FCbField ZenServerField = ManifestObject["zenserver"];
-				if (ZenServerField)
-				{
-					UE::Zen::FServiceSettings ZenServiceSettings;
-					ZenServiceSettings.ReadFromCompactBinary(ZenServerField["settings"]);
-					FString ProjectId = FString(ZenServerField["projectid"].AsString());
-					FString OplogId = FString(ZenServerField["oplogid"].AsString());
-
-					// We just want the auto launch functionality
-					UE::FZenStoreHttpClient ZenStoreClient(MoveTemp(ZenServiceSettings));
-					ZenStoreClient.InitializeReadOnly(ProjectId, OplogId);
-
-					if (!SponsorProcessIds.IsEmpty() && ZenServiceSettings.IsAutoLaunch() &&
-						ZenServiceSettings.SettingsVariant.Get<UE::Zen::FServiceAutoLaunchSettings>().bLimitProcessLifetime)
-					{
-						ZenStoreClient.GetZenServiceInstance().AddSponsorProcessIDs(SponsorProcessIds);
-					}
-
-					return 0;
-				}
-			}
-			else
-			{
-				UE_LOG(LogIoStore, Error, TEXT("Failed reading package store manifest"));
-			}
-		}
-
 		FString MarkerFilename;
 		if (FParse::Value(FCommandLine::Get(), TEXT("ProjectStore="), MarkerFilename))
 		{
@@ -9647,12 +9612,21 @@ int32 CreateIoStoreContainerFiles(const TCHAR* CmdLine)
 			if (Ar)
 			{
 				UE::Zen::FZenServiceInstance& ZenServiceInstance = UE::Zen::GetDefaultServiceInstance();
-				const UE::Zen::FServiceSettings& ZenServiceSettings = ZenServiceInstance.GetServiceSettings();
+				if (!ZenServiceInstance.IsServiceReady())
+				{
+					UE_LOG(LogIoStore, Error, TEXT("Failed to launch ZenServer."));
+					return -1;
+				}
 
+				const UE::Zen::FServiceSettings& ZenServiceSettings = ZenServiceInstance.GetServiceSettings();
 				if (!SponsorProcessIds.IsEmpty() && ZenServiceSettings.IsAutoLaunch() &&
 					ZenServiceInstance.GetServiceSettings().SettingsVariant.Get<UE::Zen::FServiceAutoLaunchSettings>().bLimitProcessLifetime)
 				{
-					ZenServiceInstance.AddSponsorProcessIDs(SponsorProcessIds);
+					if (!ZenServiceInstance.AddSponsorProcessIDs(SponsorProcessIds))
+					{
+						UE_LOG(LogIoStore, Error, TEXT("Failed to add sponsor process IDs to launched ZenServer."));
+						return -1;
+					}
 				}
 				return 0;
 			}
diff --git a/Engine/Source/Developer/Zen/Internal/ZenServerHttp.h b/Engine/Source/Developer/Zen/Internal/ZenServerHttp.h
--- a/Engine/Source/Developer/Zen/Internal/ZenServerHttp.h
+++ b/Engine/Source/Developer/Zen/Internal/ZenServerHttp.h
@@ -86,7 +86,7 @@ inline FStringView GetMimeType(EContentType Type)
 class FZenHttpRequest
 {
 public:
-	ZEN_API FZenHttpRequest(FStringView InDomain, bool bInLogErrors);
+	ZEN_API FZenHttpRequest(FStringView InDomain, bool bInLogErrors, uint32_t ConnectTimeoutMS = 0);
 	ZEN_API ~FZenHttpRequest();
 
 	/**
diff --git a/Engine/Source/Developer/Zen/Private/ZenServerHttp.cpp b/Engine/Source/Developer/Zen/Private/ZenServerHttp.cpp
--- a/Engine/Source/Developer/Zen/Private/ZenServerHttp.cpp
+++ b/Engine/Source/Developer/Zen/Private/ZenServerHttp.cpp
@@ -60,12 +60,16 @@ namespace UE::Zen {
 #endif //PLATFORM_WINDOWS
 	};
 
-	FZenHttpRequest::FZenHttpRequest(FStringView InDomain, bool bInLogErrors)
+	FZenHttpRequest::FZenHttpRequest(FStringView InDomain, bool bInLogErrors, uint32_t ConnectTimeoutMS)
 		: bLogErrors(bInLogErrors)
 		, Domain(InDomain)
 	{
 		Curl = curl_easy_init();
 		Reset();
+		if (ConnectTimeoutMS != 0)
+		{
+			curl_easy_setopt(Curl, CURLOPT_CONNECTTIMEOUT_MS, ConnectTimeoutMS);
+		}
 	}
 
 	FZenHttpRequest::~FZenHttpRequest()
@@ -90,7 +94,7 @@ namespace UE::Zen {
 		curl_easy_reset(Curl);
 
 		// Options that are always set for all connections.
-		curl_easy_setopt(Curl, CURLOPT_CONNECTTIMEOUT, 30L);
+		curl_easy_setopt(Curl, CURLOPT_CONNECTTIMEOUT, 5L);
 		curl_easy_setopt(Curl, CURLOPT_EXPECT_100_TIMEOUT_MS, 0);
 		curl_easy_setopt(Curl, CURLOPT_NOSIGNAL, 1L);
 		curl_easy_setopt(Curl, CURLOPT_DNS_CACHE_TIMEOUT, -1L); // Don't re-resolve names mid-session
diff --git a/Engine/Source/Developer/Zen/Private/ZenServerInterface.cpp b/Engine/Source/Developer/Zen/Private/ZenServerInterface.cpp
--- a/Engine/Source/Developer/Zen/Private/ZenServerInterface.cpp
+++ b/Engine/Source/Developer/Zen/Private/ZenServerInterface.cpp
@@ -853,41 +853,45 @@ FServiceSettings::ReadFromConfig()
 bool
 FServiceSettings::ReadFromCompactBinary(FCbFieldView Field)
 {
-	if (bool bAutoLaunchValue = Field["bAutoLaunch"].AsBool())
+	if (Field.IsObject())
 	{
-		if (!TryApplyAutoLaunchOverride())
+		if (bool bAutoLaunchValue = Field["bAutoLaunch"].AsBool())
 		{
-			SettingsVariant.Emplace<FServiceAutoLaunchSettings>();
-			FServiceAutoLaunchSettings& AutoLaunchSettings = SettingsVariant.Get<FServiceAutoLaunchSettings>();
-
-			if (FCbObjectView AutoLaunchSettingsObject = Field["AutoLaunchSettings"].AsObjectView())
+			if (!TryApplyAutoLaunchOverride())
 			{
-				AutoLaunchSettings.DataPath = FString(AutoLaunchSettingsObject["DataPath"].AsString());
-				AutoLaunchSettings.ExtraArgs = FString(AutoLaunchSettingsObject["ExtraArgs"].AsString());
-				AutoLaunchSettings.DesiredPort = AutoLaunchSettingsObject["DesiredPort"].AsInt16();
-				AutoLaunchSettings.bShowConsole = AutoLaunchSettingsObject["ShowConsole"].AsBool();
-				AutoLaunchSettings.bIsDefaultDataPath = AutoLaunchSettingsObject["IsDefaultDataPath"].AsBool();
-				AutoLaunchSettings.bLimitProcessLifetime = AutoLaunchSettingsObject["LimitProcessLifetime"].AsBool();
-				ApplyProcessLifetimeOverride(AutoLaunchSettings.bLimitProcessLifetime);
-				AutoLaunchSettings.bAllowRemoteNetworkService = AutoLaunchSettingsObject["AllowRemoteNetworkService"].AsBool();
-				AutoLaunchSettings.bSendUnattendedBugReports = AutoLaunchSettingsObject["SendUnattendedBugReports"].AsBool();
-				AutoLaunchSettings.bIsDefaultSharedRunContext = AutoLaunchSettingsObject["IsDefaultSharedRunContext"].AsBool(AutoLaunchSettings.bIsDefaultSharedRunContext);
-				AutoLaunchSettings.InstallMode = ZenGetInstallModeFromString(FString(AutoLaunchSettingsObject["InstallMode"].AsString()));
+				SettingsVariant.Emplace<FServiceAutoLaunchSettings>();
+				FServiceAutoLaunchSettings& AutoLaunchSettings = SettingsVariant.Get<FServiceAutoLaunchSettings>();
+
+				if (FCbObjectView AutoLaunchSettingsObject = Field["AutoLaunchSettings"].AsObjectView())
+				{
+					AutoLaunchSettings.DataPath = FString(AutoLaunchSettingsObject["DataPath"].AsString());
+					AutoLaunchSettings.ExtraArgs = FString(AutoLaunchSettingsObject["ExtraArgs"].AsString());
+					AutoLaunchSettings.DesiredPort = AutoLaunchSettingsObject["DesiredPort"].AsInt16();
+					AutoLaunchSettings.bShowConsole = AutoLaunchSettingsObject["ShowConsole"].AsBool();
+					AutoLaunchSettings.bIsDefaultDataPath = AutoLaunchSettingsObject["IsDefaultDataPath"].AsBool();
+					AutoLaunchSettings.bLimitProcessLifetime = AutoLaunchSettingsObject["LimitProcessLifetime"].AsBool();
+					ApplyProcessLifetimeOverride(AutoLaunchSettings.bLimitProcessLifetime);
+					AutoLaunchSettings.bAllowRemoteNetworkService = AutoLaunchSettingsObject["AllowRemoteNetworkService"].AsBool();
+					AutoLaunchSettings.bSendUnattendedBugReports = AutoLaunchSettingsObject["SendUnattendedBugReports"].AsBool();
+					AutoLaunchSettings.bIsDefaultSharedRunContext = AutoLaunchSettingsObject["IsDefaultSharedRunContext"].AsBool(AutoLaunchSettings.bIsDefaultSharedRunContext);
+					AutoLaunchSettings.InstallMode = ZenGetInstallModeFromString(FString(AutoLaunchSettingsObject["InstallMode"].AsString()));
+				}
 			}
 		}
-	}
-	else
-	{
-		SettingsVariant.Emplace<FServiceConnectSettings>();
-		FServiceConnectSettings& ConnectExistingSettings = SettingsVariant.Get<FServiceConnectSettings>();
-
-		if (FCbObjectView ConnectExistingSettingsObject = Field["ConnectExistingSettings"].AsObjectView())
+		else
 		{
-			ConnectExistingSettings.HostName = FString(ConnectExistingSettingsObject["HostName"].AsString());
-			ConnectExistingSettings.Port = ConnectExistingSettingsObject["Port"].AsInt16();
+			SettingsVariant.Emplace<FServiceConnectSettings>();
+			FServiceConnectSettings& ConnectExistingSettings = SettingsVariant.Get<FServiceConnectSettings>();
+
+			if (FCbObjectView ConnectExistingSettingsObject = Field["ConnectExistingSettings"].AsObjectView())
+			{
+				ConnectExistingSettings.HostName = FString(ConnectExistingSettingsObject["HostName"].AsString());
+				ConnectExistingSettings.Port = ConnectExistingSettingsObject["Port"].AsInt16();
+			}
 		}
+		return true;
 	}
-	return true;
+	return false;
 }
 
 bool
@@ -1288,6 +1292,10 @@ FZenLocalServiceRunContext::ReadFromJson(FJsonObject& JsonObject)
 	WorkingDirectory = JsonObject.Values.FindRef(TEXT("WorkingDirectory"))->AsString();
 	DataPath = JsonObject.Values.FindRef(TEXT("DataPath"))->AsString();
 	bShowConsole = JsonObject.Values.FindRef(TEXT("ShowConsole"))->AsBool();
+	if (!JsonObject.TryGetBoolField(TEXT("LimitProcessLifetime"), bLimitProcessLifetime))
+	{
+		bLimitProcessLifetime = false;
+	}
 	return true;
 }
 
@@ -1299,6 +1307,7 @@ FZenLocalServiceRunContext::WriteToJson(TJsonWriter<TCHAR, TPrettyJsonPrintPolic
 	Writer.WriteValue(TEXT("WorkingDirectory"), WorkingDirectory);
 	Writer.WriteValue(TEXT("DataPath"), DataPath);
 	Writer.WriteValue(TEXT("ShowConsole"), bShowConsole);
+	Writer.WriteValue(TEXT("LimitProcessLifetime"), bLimitProcessLifetime);
 }
 
 bool
@@ -1355,12 +1364,17 @@ IsLocalServiceRunning(const TCHAR* DataPath, uint16* OutEffectivePort)
 }
 
 FProcHandle
-StartLocalService(const FZenLocalServiceRunContext& Context, const TCHAR* TransientArgs)
+StartLocalService(const FZenLocalServiceRunContext& Context, FStringView TransientArgs)
 {
 	FString Parms = Context.GetCommandlineArguments();
-	if (TransientArgs)
+	if (Context.GetLimitProcessLifetime())
+	{
+		Parms.Appendf(TEXT(" --owner-pid %d"), FPlatformProcess::GetCurrentProcessId());
+	}
+
+	if (!TransientArgs.IsEmpty())
 	{
-		Parms.Appendf(TEXT(" %s"), TransientArgs);
+		Parms.Appendf(TEXT(" %.*s"), TransientArgs.Len(), TransientArgs.GetData());
 	}
 
 	UE_LOG(LogZenServiceInstance, Display, TEXT("Launching executable '%s', working dir '%s', data dir '%s', args '%s'"), *Context.GetExecutable(), *Context.GetWorkingDirectory(), *Context.GetDataPath(), *Parms);
@@ -1437,6 +1451,109 @@ StartLocalService(const FZenLocalServiceRunContext& Context, const TCHAR* Transi
 	return Proc;
 }
 
+bool
+StartLocalService(const FZenLocalServiceRunContext& Context)
+{
+	FString TransientParms;
+	FString StartupEventName = ZenSharedEvent::GetStartupEventName();
+	ZenSharedEvent StartupEvent(StartupEventName);
+	if (!StartupEvent.Create())
+	{
+		UE_LOG(LogZenServiceInstance, Warning, TEXT("Failed to create startup event when launcing executable '%s'"), *Context.GetExecutable());
+		return false;
+	}
+	TransientParms.Appendf(TEXT(" --child-id %s"), *StartupEventName);
+
+	FProcHandle Proc = StartLocalService(Context, TransientParms);
+	if (Proc.IsValid())
+	{
+		bool ZenServerIsReady = false;
+		FScopedSlowTask WaitForZenReadySlowTask(0, NSLOCTEXT("Zen", "Zen_WaitingForReady", "Waiting for ZenServer to be ready"));
+		uint64 ZenWaitStartTime = FPlatformTime::Cycles64();
+
+		enum class EWaitDurationPhase
+		{
+			Short,
+			Medium,
+			Long
+		} DurationPhase = EWaitDurationPhase::Short;
+
+		while (FPlatformProcess::IsProcRunning(Proc))
+		{
+			if (StartupEvent.Wait(5000))
+			{
+				ZenServerIsReady = true;
+				break;
+			}
+
+			double ZenWaitDuration = FPlatformTime::ToSeconds64(FPlatformTime::Cycles64() - ZenWaitStartTime);
+			if (ZenWaitDuration >= 10.0)
+			{
+				if (DurationPhase == EWaitDurationPhase::Short)
+				{
+					if (!FPlatformProcess::IsProcRunning(Proc))
+					{
+#if !IS_PROGRAM
+						if (!FApp::IsUnattended() && !IsRunningCommandlet() && !GIsRunningUnattendedScript)
+						{
+							FText ZenLaunchFailurePromptTitle = NSLOCTEXT("Zen", "Zen_LaunchFailurePromptTitle", "Failed to launch");
+
+							FFormatNamedArguments FormatArguments;
+							FString LogFilePath = FPaths::Combine(Context.GetDataPath(), TEXT("logs"), TEXT("zenserver.log"));
+							FPaths::MakePlatformFilename(LogFilePath);
+							FormatArguments.Add(TEXT("LogFilePath"), FText::FromString(LogFilePath));
+							FText ZenLaunchFailurePromptText = FText::Format(NSLOCTEXT("Zen", "Zen_LaunchFailurePromptText", "Unreal Zen Storage Server failed to launch. Please check the ZenServer log file for details:\n{LogFilePath}"), FormatArguments);
+							FPlatformMisc::MessageBoxExt(EAppMsgType::Ok, *ZenLaunchFailurePromptText.ToString(), *ZenLaunchFailurePromptTitle.ToString());
+							break;
+						}
+						else
+#endif
+						{
+							// Just log as there is no one to show a message
+							UE_LOG(LogZenServiceInstance, Warning, TEXT("Unreal Zen Storage Server did not launch in the expected duration"));
+							break;
+						}
+					}
+
+					// Note that the dialog may not show up when zenserver is needed early in the launch cycle, but this will at least ensure
+					// the splash screen is refreshed with the appropriate text status message.
+					WaitForZenReadySlowTask.MakeDialog(true, false);
+					UE_LOG(LogZenServiceInstance, Display, TEXT("Waiting for ZenServer to be ready..."));
+					DurationPhase = EWaitDurationPhase::Medium;
+				}
+#if !IS_PROGRAM
+				else if (!(FApp::IsUnattended() || IsRunningCommandlet() || GIsRunningUnattendedScript) && ZenWaitDuration > 20.0 && (DurationPhase == EWaitDurationPhase::Medium))
+				{
+					FText ZenLongWaitPromptTitle = NSLOCTEXT("Zen", "Zen_LongWaitPromptTitle", "Wait for ZenServer?");
+					FText ZenLongWaitPromptText = NSLOCTEXT("Zen", "Zen_LongWaitPromptText", "Unreal Zen Storage Server is taking a long time to launch. It may be performing maintenance. Keep waiting?");
+					if (FPlatformMisc::MessageBoxExt(EAppMsgType::YesNo, *ZenLongWaitPromptText.ToString(), *ZenLongWaitPromptTitle.ToString()) == EAppReturnType::No)
+					{
+						break;
+					}
+					DurationPhase = EWaitDurationPhase::Long;
+				}
+#endif
+				if (WaitForZenReadySlowTask.ShouldCancel())
+				{
+					break;
+				}
+			}
+		}
+
+		if (!ZenServerIsReady)
+		{
+			if (FPlatformProcess::IsProcRunning(Proc))
+			{
+				UE_LOG(LogZenServiceInstance, Warning, TEXT("Terminating unresponsive process for executable '%s'"), *Context.GetExecutable());
+				FPlatformProcess::TerminateProc(Proc, true);
+			}
+		}
+		FPlatformProcess::CloseProc(Proc);
+		return ZenServerIsReady;
+	}
+	return false;
+}
+	
 bool
 StopLocalService(const TCHAR* DataPath, double MaximumWaitDurationSeconds)
 {
@@ -1591,15 +1708,23 @@ FZenServiceInstance::IsServiceRunning()
 	return !Settings.IsAutoLaunch() || bHasLaunchedLocal;
 }
 
-bool 
+bool
 FZenServiceInstance::IsServiceReady()
+{
+	return PingService(2);
+}
+
+bool 
+FZenServiceInstance::PingService(uint32_t AttemptCount)
 {
 	uint32 Attempt = 0;
 	while (IsServiceRunning())
 	{
+		Attempt++;
+
 		TStringBuilder<128> ZenDomain;
 		ZenDomain << HostName << TEXT(":") << Port;
-		Zen::FZenHttpRequest Request(ZenDomain.ToString(), false);
+		Zen::FZenHttpRequest Request(ZenDomain.ToString(), false, 500);
 		Zen::FZenHttpRequest::Result Result = Request.PerformBlockingDownload(TEXTVIEW("health/ready"), nullptr, Zen::EContentType::Text);
 		
 		if (Result == Zen::FZenHttpRequest::Result::Success && Zen::IsSuccessCode(Request.GetResponseCode()))
@@ -1610,7 +1735,7 @@ FZenServiceInstance::IsServiceReady()
 
 		if (IsServiceRunningLocally())
 		{
-			if (Attempt > 4)
+			if (Attempt == AttemptCount)
 			{
 				UE_LOG(LogZenServiceInstance, Warning, TEXT("Unable to reach Unreal Zen Storage Server HTTP service at %s. Status: %d . Response: %s"), ZenDomain.ToString(), Request.GetResponseCode(), *Request.GetResponseAsString());
 				break;
@@ -1621,7 +1746,6 @@ FZenServiceInstance::IsServiceReady()
 			UE_LOG(LogZenServiceInstance, Display, TEXT("Unable to reach Unreal Zen Storage Server HTTP service at %s. Status: %d . Response: %s"), ZenDomain.ToString(), Request.GetResponseCode(), *Request.GetResponseAsString());
 			break;
 		}
-		Attempt++;
 	}
 	return false;
 }
@@ -1654,7 +1778,7 @@ FZenServiceInstance::TryRecovery()
 		}
 
 		// We test if the service is healthy as a different process might already have triggered a recovery
-		bLastRecoveryResult = IsServiceReady();
+		bLastRecoveryResult = PingService(1);
 		if (bLastRecoveryResult)
 		{
 			UE_LOG(LogZenServiceInstance, Display, TEXT("Local ZenServer status: Healthy. Skipping recovery"));
@@ -1688,7 +1812,7 @@ FZenServiceInstance::TryRecovery()
 				AutoLaunch(Settings.SettingsVariant.Get<FServiceAutoLaunchSettings>(), ExecutablePath, ExecutionContextFilePath, HostName, Port);
 
 				FDateTime StartedWaitingForHealth = FDateTime::UtcNow();
-				bLastRecoveryResult = IsServiceReady();
+				bLastRecoveryResult = PingService(1);
 				while (!bLastRecoveryResult)
 				{
 					FTimespan WaitForHealth = FDateTime::UtcNow() - StartedWaitingForHealth;
@@ -1703,7 +1827,7 @@ FZenServiceInstance::TryRecovery()
 					{
 						AutoLaunch(Settings.SettingsVariant.Get<FServiceAutoLaunchSettings>(), ExecutablePath, ExecutionContextFilePath, HostName, Port);
 					}
-					bLastRecoveryResult = IsServiceReady();
+					bLastRecoveryResult = PingService(1);
 				}
 				LastRecoveryTicks.store(FDateTime::UtcNow().GetTicks(), std::memory_order_relaxed);
 				UE_LOG(LogZenServiceInstance, Display, TEXT("Local ZenServer recovery finished."));
@@ -1786,7 +1910,7 @@ FZenServiceInstance::Initialize()
 							}
 							AutoLaunchedPort = Port;
 							bIsRunningLocally = true;
-							if (IsServiceReady())
+							if (PingService(2))
 							{
 								break;
 							}
@@ -2122,6 +2246,7 @@ FZenServiceInstance::AutoLaunch(const FServiceAutoLaunchSettings& InSettings, co
 				DesiredRunContext.WorkingDirectory = WorkingDirectory;
 				DesiredRunContext.DataPath = InSettings.DataPath;
 				DesiredRunContext.bShowConsole = InSettings.bShowConsole;
+				DesiredRunContext.bLimitProcessLifetime = InSettings.bLimitProcessLifetime;
 
 				FZenLocalServiceRunContext CurrentRunContext;
 
@@ -2226,12 +2351,6 @@ FZenServiceInstance::AutoLaunch(const FServiceAutoLaunchSettings& InSettings, co
 		}
 
 		FString ParmsWithoutTransients = DetermineCmdLineWithoutTransientComponents(InSettings, InSettings.DesiredPort);
-		FString TransientParms;
-
-		if (InSettings.bLimitProcessLifetime)
-		{
-			TransientParms.Appendf(TEXT("--owner-pid %d"), FPlatformProcess::GetCurrentProcessId());
-		}
 
 		FZenLocalServiceRunContext EffectiveRunContext;
 		EffectiveRunContext.Executable = ExecutablePath;
@@ -2239,106 +2358,15 @@ FZenServiceInstance::AutoLaunch(const FServiceAutoLaunchSettings& InSettings, co
 		EffectiveRunContext.WorkingDirectory = WorkingDirectory;
 		EffectiveRunContext.DataPath = InSettings.DataPath;
 		EffectiveRunContext.bShowConsole = InSettings.bShowConsole;
+		EffectiveRunContext.bLimitProcessLifetime = InSettings.bLimitProcessLifetime;
 
-		FString StartupEventName = ZenSharedEvent::GetStartupEventName();
-		ZenSharedEvent StartupEvent(StartupEventName);
-		if (!StartupEvent.Create())
-		{
-			return false;
-		}
-		TransientParms.Appendf(TEXT(" --child-id %s"), *StartupEventName);
-
-		FProcHandle Proc = StartLocalService(EffectiveRunContext, TransientParms.IsEmpty() ? nullptr : *TransientParms);
-
-		if (Proc.IsValid())
+		if (StartLocalService(EffectiveRunContext))
 		{
 			// Only write run context if we're using the default shared run context
 			if (InSettings.bIsDefaultSharedRunContext)
 			{
 				EffectiveRunContext.WriteToJsonFile(*ExecutionContextFilePath);
 			}
-
-			bool ZenServerIsReady = false;
-			FScopedSlowTask WaitForZenReadySlowTask(0, NSLOCTEXT("Zen", "Zen_WaitingForReady", "Waiting for ZenServer to be ready"));
-			uint64 ZenWaitStartTime = FPlatformTime::Cycles64();
-
-			enum class EWaitDurationPhase
-			{
-				Short,
-				Medium,
-				Long
-			} DurationPhase = EWaitDurationPhase::Short;
-
-			while (FPlatformProcess::IsProcRunning(Proc))
-			{
-				if (StartupEvent.Wait(5000))
-				{
-					ZenServerIsReady = true;
-					break;
-				}
-
-				double ZenWaitDuration = FPlatformTime::ToSeconds64(FPlatformTime::Cycles64() - ZenWaitStartTime);
-				if (ZenWaitDuration >= 10.0)
-				{
-					if (DurationPhase == EWaitDurationPhase::Short)
-					{
-						if (!FPlatformProcess::IsProcRunning(Proc))
-						{
-#if !IS_PROGRAM
-							if (!FApp::IsUnattended() && !IsRunningCommandlet() && !GIsRunningUnattendedScript)
-							{
-								FText ZenLaunchFailurePromptTitle = NSLOCTEXT("Zen", "Zen_LaunchFailurePromptTitle", "Failed to launch");
-
-								FFormatNamedArguments FormatArguments;
-								FString LogFilePath = FPaths::Combine(InSettings.DataPath, TEXT("logs"), TEXT("zenserver.log"));
-								FPaths::MakePlatformFilename(LogFilePath);
-								FormatArguments.Add(TEXT("LogFilePath"), FText::FromString(LogFilePath));
-								FText ZenLaunchFailurePromptText = FText::Format(NSLOCTEXT("Zen", "Zen_LaunchFailurePromptText", "Unreal Zen Storage Server failed to launch. Please check the ZenServer log file for details:\n{LogFilePath}"), FormatArguments);
-								FPlatformMisc::MessageBoxExt(EAppMsgType::Ok, *ZenLaunchFailurePromptText.ToString(), *ZenLaunchFailurePromptTitle.ToString());
-								break;
-							}
-							else
-#endif
-							{
-								// Just log as there is no one to show a message
-								UE_LOG(LogZenServiceInstance, Warning, TEXT("Unreal Zen Storage Server did not launch in the expected duration"));
-								break;
-							}
-						}
-
-						// Note that the dialog may not show up when zenserver is needed early in the launch cycle, but this will at least ensure
-						// the splash screen is refreshed with the appropriate text status message.
-						WaitForZenReadySlowTask.MakeDialog(true, false);
-						UE_LOG(LogZenServiceInstance, Display, TEXT("Waiting for ZenServer to be ready..."));
-						DurationPhase = EWaitDurationPhase::Medium;
-					}
-#if !IS_PROGRAM
-					else if (!(FApp::IsUnattended() || IsRunningCommandlet() || GIsRunningUnattendedScript) && ZenWaitDuration > 20.0 && (DurationPhase == EWaitDurationPhase::Medium))
-					{
-						FText ZenLongWaitPromptTitle = NSLOCTEXT("Zen", "Zen_LongWaitPromptTitle", "Wait for ZenServer?");
-						FText ZenLongWaitPromptText = NSLOCTEXT("Zen", "Zen_LongWaitPromptText", "Unreal Zen Storage Server is taking a long time to launch. It may be performing maintenance. Keep waiting?");
-						if (FPlatformMisc::MessageBoxExt(EAppMsgType::YesNo, *ZenLongWaitPromptText.ToString(), *ZenLongWaitPromptTitle.ToString()) == EAppReturnType::No)
-						{
-							break;
-						}
-						DurationPhase = EWaitDurationPhase::Long;
-					}
-#endif
-					if (WaitForZenReadySlowTask.ShouldCancel())
-					{
-						break;
-					}
-				}
-			}
-
-			if (!ZenServerIsReady)
-			{
-				if (FPlatformProcess::IsProcRunning(Proc))
-				{
-					UE_LOG(LogZenServiceInstance, Warning, TEXT("Terminating unresponsive process for executable '%s' on port %u"), *ExecutablePath, InSettings.DesiredPort);
-					FPlatformProcess::TerminateProc(Proc, true);
-				}
-			}
 		}
 		else
 		{
@@ -2356,7 +2384,7 @@ FZenServiceInstance::AutoLaunch(const FServiceAutoLaunchSettings& InSettings, co
 		}
 		else if (!RunningEntry->AddSponsorProcess(FPlatformProcess::GetCurrentProcessId()))
 		{
-			UE_LOG(LogZenServiceInstance, Warning, TEXT("Failed attach as sponsor process to executable '%s' on port %u, too many sponsored processes attached already"), *ExecutablePath, InSettings.DesiredPort);
+			UE_LOG(LogZenServiceInstance, Warning, TEXT("Failed attach as sponsor process to executable '%s' on port %u, too many sponsored processes attached already or zenserver is unresponsive"), *ExecutablePath, InSettings.DesiredPort);
 		}
 	}
 
diff --git a/Engine/Source/Developer/Zen/Private/ZenServerState.cpp b/Engine/Source/Developer/Zen/Private/ZenServerState.cpp
--- a/Engine/Source/Developer/Zen/Private/ZenServerState.cpp
+++ b/Engine/Source/Developer/Zen/Private/ZenServerState.cpp
@@ -401,24 +401,41 @@ const ZenServerState::ZenServerEntry* ZenServerState::LookupByPid(uint32 Pid) co
 bool
 ZenServerState::ZenServerEntry::AddSponsorProcess(uint32 PidToAdd)
 {
-	for (std::atomic<uint32>& PidEntry : SponsorPids)
-	{
-		if (PidEntry.load(std::memory_order_relaxed) == 0)
+	uint32_t ServerPid = Pid.load();
+
+	// Sponsor processes are checked every second, so 2 second wait time should be enough
+	const FTimespan MaximumWaitForPickup = FTimespan::FromSeconds(2);
+	auto	 WaitForPickup = [&](uint32_t AddedSlotIndex) {
+		FDateTime StartedWaitingForPickup = FDateTime::UtcNow();
+		while (SponsorPids[AddedSlotIndex] == PidToAdd)
 		{
-			uint32 Expected = 0;
-			if (PidEntry.compare_exchange_strong(Expected, PidToAdd))
+			FTimespan WaitForPickup = FDateTime::UtcNow() - StartedWaitingForPickup;
+			if (WaitForPickup > MaximumWaitForPickup)
+			{
+				SponsorPids[AddedSlotIndex].compare_exchange_strong(PidToAdd, 0);
+				return false;
+			}
+			if (!IsProcessRunning(ServerPid))
 			{
-				// Success!
-				return true;
+				SponsorPids[AddedSlotIndex].compare_exchange_strong(PidToAdd, 0);
+				return false;
 			}
+			FPlatformProcess::Sleep(0.1f);
 		}
-		else if (PidEntry.load(std::memory_order_relaxed) == PidToAdd)
+		return true;
+	};
+	for (uint32_t SponsorIndex = 0; SponsorIndex < 8; SponsorIndex++)
+	{
+		if (SponsorPids[SponsorIndex].load(std::memory_order_relaxed) == PidToAdd)
 		{
-			// Success, the because pid is already in the list
-			return true;
+			return WaitForPickup(SponsorIndex);
+		}
+		uint32_t Expected = 0;
+		if (SponsorPids[SponsorIndex].compare_exchange_strong(Expected, PidToAdd))
+		{
+			return WaitForPickup(SponsorIndex);
 		}
 	}
-
 	return false;
 }
 
diff --git a/Engine/Source/Developer/Zen/Public/Experimental/ZenServerInterface.h b/Engine/Source/Developer/Zen/Public/Experimental/ZenServerInterface.h
--- a/Engine/Source/Developer/Zen/Public/Experimental/ZenServerInterface.h
+++ b/Engine/Source/Developer/Zen/Public/Experimental/ZenServerInterface.h
@@ -97,10 +97,11 @@ public:
 	FString GetWorkingDirectory() const { return WorkingDirectory; }
 	FString GetDataPath() const { return DataPath; }
 	bool GetShowConsole() const { return bShowConsole; }
+	bool GetLimitProcessLifetime() const { return bLimitProcessLifetime; }
 
 	friend inline bool operator==(FZenLocalServiceRunContext LHS, FZenLocalServiceRunContext RHS)
 	{
-		return (LHS.Executable == RHS.Executable) && (LHS.CommandlineArguments == RHS.CommandlineArguments) && (LHS.WorkingDirectory == RHS.WorkingDirectory) && (LHS.DataPath == RHS.DataPath);
+		return (LHS.Executable == RHS.Executable) && (LHS.CommandlineArguments == RHS.CommandlineArguments) && (LHS.WorkingDirectory == RHS.WorkingDirectory) && (LHS.DataPath == RHS.DataPath) && (LHS.bLimitProcessLifetime == RHS.bLimitProcessLifetime);
 	}
 
 private:
@@ -108,7 +109,8 @@ private:
 	FString CommandlineArguments;
 	FString WorkingDirectory;
 	FString DataPath;
-	bool bShowConsole;
+	bool bShowConsole = false;
+	bool bLimitProcessLifetime = false;
 
 	bool ReadFromJson(FJsonObject& JsonObject);
 	void WriteToJson(TJsonWriter<TCHAR, TPrettyJsonPrintPolicy<TCHAR>>& Writer) const;
@@ -121,7 +123,7 @@ private:
 };
 
 UE_API bool IsLocalServiceRunning(const TCHAR* DataPath, uint16* OutPort = nullptr);
-UE_API FProcHandle StartLocalService(const FZenLocalServiceRunContext& Context, const TCHAR* TransientArgs = nullptr);
+UE_API bool StartLocalService(const FZenLocalServiceRunContext& Context);
 UE_API bool StopLocalService(const TCHAR* DataPath, double MaximumWaitDurationSeconds = 25.0);
 
 UE_API FString GetLocalServiceInstallPath();
@@ -212,7 +214,8 @@ private:
 	void Initialize();
 	bool ConditionalUpdateLocalInstall(FServiceAutoLaunchSettings::EInstallMode InstallMode);
 	static bool AutoLaunch(const FServiceAutoLaunchSettings& InSettings, const FString& ExecutablePath, const FString& ExecutionContextFilePath, FString& OutHostName, uint16& OutPort);
-	
+	bool PingService(uint32_t AttemptCount);
+
 	mutable TFuture<TOptional<FGCStatus>> GCStatusRequest;
 	mutable TOptional<FGCStatus> LastGCStatus;
 	mutable uint64 LastGCStatusTime = 0;
diff --git a/Engine/Source/Programs/AutomationTool/Scripts/CopyBuildToStagingDirectory.Automation.cs b/Engine/Source/Programs/AutomationTool/Scripts/CopyBuildToStagingDirectory.Automation.cs
--- a/Engine/Source/Programs/AutomationTool/Scripts/CopyBuildToStagingDirectory.Automation.cs
+++ b/Engine/Source/Programs/AutomationTool/Scripts/CopyBuildToStagingDirectory.Automation.cs
@@ -1013,40 +1013,6 @@ namespace AutomationScripts
 			{
 				SC.PackageStoreData = new PackageStoreData();
 				SC.PackageStoreData.ManifestFullPath = PackageStoreManifestFile.FullName;
-
-				byte[] ManifestData = File.ReadAllBytes(PackageStoreManifestFile.FullName);
-				CbObject ManifestObject = new CbField(ManifestData).AsObject();
-				CbObject ZenServerObject = ManifestObject["zenserver"].AsObject();
-				if (ZenServerObject != CbObject.Empty)
-				{
-					Logger.LogInformation("Reading oplog from Zen...");
-					string Host = "localhost";
-					bool IsLocalHost = true;
-					int Port = 8558;
-					string ProjectId = ZenServerObject["projectid"].AsString();
-					string OplogId = ZenServerObject["oplogid"].AsString();
-					bool bAutoLaunch = false;
-					CbObject SettingsObject = ZenServerObject["settings"].AsObject();
-					if (SettingsObject != CbObject.Empty)
-					{
-						bAutoLaunch = SettingsObject["bAutoLaunch"].AsBool();
-						CbObject AutoLaunchSettingsObject = SettingsObject["AutoLaunchSettings"].AsObject();
-						if (AutoLaunchSettingsObject != CbObject.Empty)
-						{
-							Port = AutoLaunchSettingsObject["DesiredPort"].AsInt16();
-						}
-						CbObject ConnectExistingSettingsObject = SettingsObject["ConnectExistingSettings"].AsObject();
-						if (ConnectExistingSettingsObject != CbObject.Empty)
-						{
-							Host = ConnectExistingSettingsObject["HostName"].AsString();
-							IsLocalHost = Host == "localhost" || Host == "127.0.0.1" || Host == "::1";
-							Port = ConnectExistingSettingsObject["Port"].AsInt16();
-						}
-					}
-
-					SC.PackageStoreData.ZenServerStore = new ZenServerStoreData { HostName = Host, HostPort = Port, ProjectId = ProjectId, OplogId = OplogId, IsLocalHost = IsLocalHost };
-					SC.PackageStoreData.ZenCookedFiles = ReadZenCookedFilesFromZenServer(Host, Port, ProjectId, OplogId, Params, SC, bAutoLaunch, "PackageStoreManifest", PackageStoreManifestFile.FullName);
-				}
 				return;
 			}
 
