Fix proc stat file null termination zenserver state GetPidStatus
Fix errno checking zenserver state IsLockFileLocked

diff --git a/Engine/Source/Developer/Zen/Private/ZenServerState.cpp b/Engine/Source/Developer/Zen/Private/ZenServerState.cpp
--- a/Engine/Source/Developer/Zen/Private/ZenServerState.cpp
+++ b/Engine/Source/Developer/Zen/Private/ZenServerState.cpp
@@ -42,7 +42,7 @@ GetPidStatus(int Pid)
 		fclose(StatFile);
 		if (Size > 0)
 		{
-			Buffer[Size + 1] = 0;
+			Buffer[Size] = 0;
 			char* ScanPtr = strrchr(Buffer, ')');
 			if (ScanPtr && ScanPtr[1] != '\0')
 			{
@@ -800,8 +800,9 @@ ZenLockFileData::IsLockFileLocked(const TCHAR* FileName, bool bAttemptCleanUp)
 	int32 LockRet = flock(Fd, LOCK_EX | LOCK_NB);
 	if (LockRet < 0)
 	{
+		int LastError = errno;
 		close(Fd);
-		return errno == EWOULDBLOCK || errno == EAGAIN;
+		return LastError == EWOULDBLOCK || LastError == EAGAIN;
 	}
 
 	// Consider the lock file as orphaned if we we managed to claim the lock for
@@ -879,7 +880,8 @@ ZenLockFileData::ReadCbLockFile(const TCHAR* FileName)
 		return {};
 	}
 
-	if (errno != EWOULDBLOCK && errno != EAGAIN)
+	int LastError = errno;
+	if (LastError != EWOULDBLOCK && LastError != EAGAIN)
 	{
 		close(Fd);
 		return {};
