Make sure we follow the same logic for IsLockFileLocked and ReadCbLockFile on Linux/Mac as we do on Windows
Pick up executable path from lockfile when available
Fix and extend logging in ZenServerState.cpp

diff --git a/Engine/Source/Developer/Zen/Private/ZenServerInterface.cpp b/Engine/Source/Developer/Zen/Private/ZenServerInterface.cpp
--- a/Engine/Source/Developer/Zen/Private/ZenServerInterface.cpp
+++ b/Engine/Source/Developer/Zen/Private/ZenServerInterface.cpp
@@ -1161,7 +1161,7 @@ static bool ShutDownZenServerProcessLockingDataDir(const FString& DataPath, doub
 			return true;
 		}
 		uint32_t Pid = 0;
-		if (!ZenServerState::FindRunningProcessId(*GetServiceExecutableName(), &Pid))
+		if (!ZenServerState::FindRunningProcessId(LockFileState.ExecutablePath.IsEmpty() ? *GetServiceExecutableName() : *LockFileState.ExecutablePath, &Pid))
 		{
 			if (!ZenLockFileData::IsLockFileLocked(*LockFilePath, true))
 			{
diff --git a/Engine/Source/Developer/Zen/Private/ZenServerState.cpp b/Engine/Source/Developer/Zen/Private/ZenServerState.cpp
--- a/Engine/Source/Developer/Zen/Private/ZenServerState.cpp
+++ b/Engine/Source/Developer/Zen/Private/ZenServerState.cpp
@@ -197,10 +197,10 @@ bool ZenServerState::Terminate(uint32 Pid)
 	int Res = kill(pid_t(Pid), SIGKILL);
 	if (Res != 0)
 	{
-		int err = errno;
-		if (err != ESRCH)
+		int LastError = errno;
+		if (LastError != ESRCH)
 		{
-			UE_LOG(LogZenServiceState, Warning, TEXT("Failed to terminate running process %d: %d"), Pid, err);
+			UE_LOG(LogZenServiceState, Warning, TEXT("Failed to terminate running process %d: %d"), Pid, LastError);
 			return false;
 		}
 	}
@@ -458,7 +458,7 @@ bool ZenSharedEvent::Create()
 	if (m_EventHandle == NULL)
 	{
 		DWORD LastError = GetLastError();
-		UE_LOG(LogZenServiceState, Warning, TEXT("Failed creating named event '{EventName}' (err: {Error})"), *FullEventName, LastError);
+		UE_LOG(LogZenServiceState, Warning, TEXT("Failed creating named event '%s' (err: %d)"), *FullEventName, LastError);
 		return false;
 	}
 	return true;
@@ -471,7 +471,7 @@ bool ZenSharedEvent::Create()
 	if (m_Fd < 0)
 	{
 		int LastError = errno;
-		UE_LOG(LogZenServiceState, Warning, TEXT("Failed to create named event '{EventName}' (err: {Error})"), *EventPath, LastError);
+		UE_LOG(LogZenServiceState, Warning, TEXT("Failed to create named event '%hs' (err: %d)"), *EventPath, LastError);
 		return false;
 	}
 	fchmod(m_Fd, 0666);
@@ -483,7 +483,7 @@ bool ZenSharedEvent::Create()
 		int LastError = errno;
 		close(m_Fd);
 		m_Fd = -1;
-		UE_LOG(LogZenServiceState, Warning, TEXT("Failed to create an SysV IPC key for named event '{EventName}' (err: {Error})"), *EventPath, LastError);
+		UE_LOG(LogZenServiceState, Warning, TEXT("Failed to create an SysV IPC key for named event '%hs' (err: %d)"), *EventPath, LastError);
 		return false;
 	}
 
@@ -494,7 +494,7 @@ bool ZenSharedEvent::Create()
 		int LastError = errno;
 		close(m_Fd);
 		m_Fd = -1;
-		UE_LOG(LogZenServiceState, Warning, TEXT("Failed creating an SysV semaphore for named event '{EventName}' (err: {Error})"), *EventPath, LastError);
+		UE_LOG(LogZenServiceState, Warning, TEXT("Failed creating an SysV semaphore for named event '%hs' (err: %d)"), *EventPath, LastError);
 		return false;
 	}
 
@@ -523,7 +523,7 @@ bool ZenSharedEvent::Exists() const
 		DWORD LastError = GetLastError();
 		if (LastError != ERROR_FILE_NOT_FOUND)
 		{
-			UE_LOG(LogZenServiceState, Warning, TEXT("Failed checking existance of named event '{EventName}' (err: {Error})"), *FullEventName, LastError);
+			UE_LOG(LogZenServiceState, Warning, TEXT("Failed checking existance of named event '%s' (err: %d)"), *FullEventName, LastError);
 		}
 		return false;
 	}
@@ -538,7 +538,7 @@ bool ZenSharedEvent::Exists() const
 	if (IpcKey < 0)
 	{
 		int LastError = errno;
-		UE_LOG(LogZenServiceState, Warning, TEXT("Failed to create an SysV IPC key for named event '{EventName}' (err: {Error})"), *EventPath, LastError);
+		UE_LOG(LogZenServiceState, Warning, TEXT("Failed to create an SysV IPC key for named event '%hs' (err: %d)"), *EventPath, LastError);
 		return false;
 	}
 	int Semaphore = semget(IpcKey, 1, 0400);
@@ -547,7 +547,7 @@ bool ZenSharedEvent::Exists() const
 		int LastError = errno;
 		if (LastError != ENOENT)
 		{
-			UE_LOG(LogZenServiceState, Warning, TEXT("Failed checking named event '{EventName}' (err: {Error})"), *EventPath, LastError);
+			UE_LOG(LogZenServiceState, Warning, TEXT("Failed checking named event '%hs' (err: %d)"), *EventPath, LastError);
 		}
 		return false;
 	}
@@ -566,7 +566,7 @@ bool ZenSharedEvent::Open()
 		DWORD LastError = GetLastError();
 		if (LastError != ERROR_FILE_NOT_FOUND)
 		{
-			UE_LOG(LogZenServiceState, Warning, TEXT("Failed opening named event '{EventName}' (err: {Error})"), *FullEventName, LastError);
+			UE_LOG(LogZenServiceState, Warning, TEXT("Failed opening named event '%s' (err: %d)"), *FullEventName, LastError);
 		}
 		return false;
 	}
@@ -579,7 +579,7 @@ bool ZenSharedEvent::Open()
 	if (IpcKey < 0)
 	{
 		int LastError = errno;
-		UE_LOG(LogZenServiceState, Warning, TEXT("Failed to create an SysV IPC key for named event '{EventName}' (err: {Error})"), *EventPath, LastError);
+		UE_LOG(LogZenServiceState, Warning, TEXT("Failed to create an SysV IPC key for named event '%hs' (err: %d)"), *EventPath, LastError);
 		return false;
 	}
 	m_Semaphore = semget(IpcKey, 1, 0600);
@@ -588,7 +588,7 @@ bool ZenSharedEvent::Open()
 		int LastError = errno;
 		if (LastError != ENOENT)
 		{
-			UE_LOG(LogZenServiceState, Warning, TEXT("Failed opening named event '{EventName}' (err: {Error})"), *EventPath, LastError);
+			UE_LOG(LogZenServiceState, Warning, TEXT("Failed opening named event '%hs' (err: %d)"), *EventPath, LastError);
 		}
 		return false;
 	}
@@ -608,7 +608,7 @@ bool ZenSharedEvent::Wait(int TimeoutMs)
 	if (Result == WAIT_FAILED)
 	{
 		DWORD LastError = GetLastError();
-		UE_LOG(LogZenServiceState, Warning, TEXT("Failed waiting for named event '{EventName}' (err: {Error})"), *m_EventName, LastError);
+		UE_LOG(LogZenServiceState, Warning, TEXT("Failed waiting for named event '%s' (err: %d)"), *m_EventName, LastError);
 		return false;
 	}
 
@@ -626,7 +626,7 @@ bool ZenSharedEvent::Wait(int TimeoutMs)
 		if (Result != 0)
 		{
 			int LastError = errno;
-			UE_LOG(LogZenServiceState, Warning, TEXT("Failed waiting for named event '{EventName}' (err: {Error})"), *m_EventName, LastError);
+			UE_LOG(LogZenServiceState, Warning, TEXT("Failed waiting for named event '%s' (err: %d)"), *m_EventName, LastError);
 			return false;
 		}
 		return true;
@@ -645,7 +645,7 @@ bool ZenSharedEvent::Wait(int TimeoutMs)
 	int LastError = errno;
 	if (LastError != EAGAIN)
 	{
-		UE_LOG(LogZenServiceState, Warning, TEXT("Failed waiting for named event '{EventName}' (err: {Error})"), *m_EventName, LastError);
+		UE_LOG(LogZenServiceState, Warning, TEXT("Failed waiting for named event '%s' (err: %d)"), *m_EventName, LastError);
 	}
 	return false;
 #elif PLATFORM_MAC
@@ -663,7 +663,7 @@ bool ZenSharedEvent::Wait(int TimeoutMs)
 			int LastError = errno;
 			if (errno != EAGAIN)
 			{
-				UE_LOG(LogZenServiceState, Warning, TEXT("Failed waiting for named event '{EventName}' (err: {Error})"), *m_EventName, LastError);
+				UE_LOG(LogZenServiceState, Warning, TEXT("Failed waiting for named event '%s' (err: %d)"), *m_EventName, LastError);
 				break;
 			}
 		}
@@ -686,7 +686,7 @@ bool ZenSharedEvent::Set()
 	else
 	{
 		DWORD LastError = GetLastError();
-		UE_LOG(LogZenServiceState, Warning, TEXT("Failed signalling named event '{EventName}' (err: {Error})"), *m_EventName, LastError);
+		UE_LOG(LogZenServiceState, Warning, TEXT("Failed signalling named event '%s' (err: %d)"), *m_EventName, LastError);
 		return false;
 	}
 #elif PLATFORM_UNIX || PLATFORM_MAC
@@ -698,7 +698,7 @@ bool ZenSharedEvent::Set()
 	else
 	{
 		int LastError = errno;
-		UE_LOG(LogZenServiceState, Warning, TEXT("Failed signalling named event '{EventName}' (err: {Error})"), *m_EventName, LastError);
+		UE_LOG(LogZenServiceState, Warning, TEXT("Failed signalling named event '%s' (err: %d)"), *m_EventName, LastError);
 		return false;
 	}
 #endif
@@ -768,15 +768,26 @@ static ZenLockFileData ReadLockData(FUniqueBuffer&& FileBytes)
 	if (ValidateCompactBinary(FileBytes, ECbValidateMode::Default) == ECbValidateError::None)
 	{
 		FCbObject LockObject(FileBytes.MoveToShared());
+
 		int32 ProcessId = LockObject["pid"].AsInt32();
-		FUtf8StringView DataDir = LockObject["data"].AsString();
-		int32 EffectivePort = LockObject["port"].AsInt32();
 		FCbObjectId SessionId = LockObject["session_id"].AsObjectId();
+		int32 EffectivePort = LockObject["port"].AsInt32();
 		bool IsReady = LockObject["ready"].AsBool();
+		FUtf8StringView DataDir = LockObject["data"].AsString();
+		FUtf8StringView ExecutablePath = LockObject["executable"].AsString();
+
 		bool IsValid = ProcessId > 0 && EffectivePort > 0 && EffectivePort <= 0xffff;
-		return ZenLockFileData{ static_cast<uint32>(ProcessId), FString(DataDir), static_cast<uint16>(EffectivePort), SessionId, IsReady, IsValid };
+
+		return ZenLockFileData{
+			.ProcessId  = static_cast<uint32>(ProcessId),
+			.SessionId = SessionId,
+			.EffectivePort = static_cast<uint16>(EffectivePort),
+			.IsReady = IsReady,
+			.DataDir = FString(DataDir),
+			.ExecutablePath = FString(ExecutablePath),
+			.IsValid = IsValid };
 	}
-	return ZenLockFileData{ 0, {}, 0, {}, false };
+	return ZenLockFileData{};
 }
 
 bool
@@ -791,32 +802,55 @@ ZenLockFileData::IsLockFileLocked(const TCHAR* FileName, bool bAttemptCleanUp)
 #elif PLATFORM_UNIX || PLATFORM_MAC
 	TAnsiStringBuilder<256> LockFilePath;
 	LockFilePath << FileName;
-	int32 Fd = open(LockFilePath.ToString(), O_RDONLY);
-	if (Fd < 0)
-	{
-		return false;
-	}
 
-	int32 LockRet = flock(Fd, LOCK_EX | LOCK_NB);
-	if (LockRet < 0)
+	if (bAttemptCleanUp)
 	{
-		int LastError = errno;
-		close(Fd);
-		return LastError == EWOULDBLOCK || LastError == EAGAIN;
-	}
+		int32 Fd = open(LockFilePath.ToString(), O_RDONLY | O_CLOEXEC);
+		if (Fd < 0)
+		{
+			int LastError = errno;
+			if (LastError == ENOENT)
+			{
+				return false;
+			}
+			UE_LOG(LogZenServiceState, Warning, TEXT("Failed opening lock file '%hs' (err: %d)"), *LockFilePath, LastError);
+			return true;
+		}
 
-	// Consider the lock file as orphaned if we we managed to claim the lock for
-	// it. Might as well delete it while we own it.
-	unlink(LockFilePath.ToString());
+		int32 LockRet = flock(Fd, LOCK_EX | LOCK_NB);
+		if (LockRet < 0)
+		{
+			int LastError = errno;
+			close(Fd);
+			if (LastError != EWOULDBLOCK)
+			{
+				UE_LOG(LogZenServiceState, Warning, TEXT("Failed locking lock file '%hs' (err: %d)"), *LockFilePath, LastError);
+			}
+			return true;
+		}
 
-	flock(Fd, LOCK_UN);
-	close(Fd);
+		unlink(LockFilePath.ToString());
 
-	return false;
+		flock(Fd, LOCK_UN);
+		close(Fd);
+	}
+
+	struct stat   Stat;
+	int Res = stat(LockFilePath.ToString(), &Stat);
+	if (Res == 0)
+	{
+		return true;
+	}
+	int LastError = errno;
+	if (LastError == ENOENT)
+	{
+		return false;
+	}
+	UE_LOG(LogZenServiceState, Warning, TEXT("Failed checking stat of '%hs' (err: %d)"), *LockFilePath, LastError);
+	return true;
 #endif
 }
 
-
 ZenLockFileData
 ZenLockFileData::ReadCbLockFile(const TCHAR* FileName)
 {
@@ -841,6 +875,7 @@ ZenLockFileData::ReadCbLockFile(const TCHAR* FileName)
 	{
 		FullFileNameBuilder.Prepend(TEXTVIEW("\\\\?\\"));
 	}
+
 	HANDLE Handle = CreateFileW(FullFileNameBuilder.ToString(), Access, WinFlags, NULL, Create, FILE_ATTRIBUTE_NORMAL, NULL);
 	if (Handle != INVALID_HANDLE_VALUE)
 	{
@@ -850,11 +885,14 @@ ZenLockFileData::ReadCbLockFile(const TCHAR* FileName)
 		{
 			checkf(LI.QuadPart == LI.u.LowPart, TEXT("Lock file exceeds supported 2GB limit."));
 			int32 FileSize32 = LI.u.LowPart;
-			FUniqueBuffer FileBytes = FUniqueBuffer::Alloc(FileSize32);
-			DWORD ReadBytes = 0;
-			if (ReadFile(Handle, FileBytes.GetData(), FileSize32, &ReadBytes, NULL) && (ReadBytes == FileSize32))
+			if (FileSize32 != 0)
 			{
-				return ReadLockData(std::move(FileBytes));
+				FUniqueBuffer FileBytes = FUniqueBuffer::Alloc(FileSize32);
+				DWORD ReadBytes = 0;
+				if (ReadFile(Handle, FileBytes.GetData(), FileSize32, &ReadBytes, NULL) && (ReadBytes == FileSize32))
+				{
+					return ReadLockData(std::move(FileBytes));
+				}
 			}
 		}
 	}
@@ -862,43 +900,26 @@ ZenLockFileData::ReadCbLockFile(const TCHAR* FileName)
 #elif PLATFORM_UNIX || PLATFORM_MAC
 	TAnsiStringBuilder<256> LockFilePath;
 	LockFilePath << FileName;
-	int32 Fd = open(LockFilePath.ToString(), O_RDONLY);
-	if (Fd < 0)
+	int32 Fd = open(LockFilePath.ToString(), O_RDONLY | O_CLOEXEC);
+	if (Fd != -1)
 	{
-		return {};
-	}
-
-	// If we can claim the lock then it's an orphaned lock file and should be
-	// ignored. Not ideal as there's a period of time when the lock can be
-	// held unncessarily.
-	int32 LockRet = flock(Fd, LOCK_EX | LOCK_NB);
-	if (LockRet >= 0)
-	{
-		unlink(LockFilePath.ToString());
-		flock(Fd, LOCK_UN);
-		close(Fd);
-		return {};
-	}
-
-	int LastError = errno;
-	if (LastError != EWOULDBLOCK && LastError != EAGAIN)
-	{
-		close(Fd);
-		return {};
-	}
-
-	struct stat Stat;
-	fstat(Fd, &Stat);
-	uint64 FileSize = uint64(Stat.st_size);
-
-	FUniqueBuffer FileBytes = FUniqueBuffer::Alloc(FileSize);
-	if (read(Fd, FileBytes.GetData(), FileSize) == FileSize)
-	{
-		close(Fd);
-		return ReadLockData(std::move(FileBytes));
+		ON_SCOPE_EXIT{ close(Fd); };
+		struct stat Stat;
+		int Res = fstat(Fd, &Stat);
+		if (Res == 0)
+		{
+			uint64 FileSize = uint64(Stat.st_size);
+			checkf(FileSize < 2u * 1024u * 1024u * 1024u, TEXT("Lock file exceeds supported 2GB limit."));
+			if (FileSize != 0)
+			{
+				FUniqueBuffer FileBytes = FUniqueBuffer::Alloc(FileSize);
+				if (read(Fd, FileBytes.GetData(), FileSize) == FileSize)
+				{
+					return ReadLockData(std::move(FileBytes));
+				}
+			}
+		}
 	}
-
-	close(Fd);
 	return {};
 #endif
 }
diff --git a/Engine/Source/Developer/Zen/Private/ZenServerState.h b/Engine/Source/Developer/Zen/Private/ZenServerState.h
--- a/Engine/Source/Developer/Zen/Private/ZenServerState.h
+++ b/Engine/Source/Developer/Zen/Private/ZenServerState.h
@@ -93,10 +93,11 @@ private:
 struct ZenLockFileData
 {
 	uint32 ProcessId = 0;
-	FString DataDir;
-	uint16 EffectivePort = 0;
 	FCbObjectId SessionId;
+	uint16 EffectivePort = 0;
 	bool IsReady = false;
+	FString DataDir;
+	FString ExecutablePath;
 	bool IsValid = false;
 
 	static bool IsLockFileLocked(const TCHAR* FileName, bool bAttemptCleanUp = false);
