Don't use MonitoredUtilityProcess.GetDuration() when launcing zen executable - it is not reliable when using multithreading.

#rnx

diff --git a/Engine/Source/Developer/Zen/Private/ZenServerInterface.cpp b/Engine/Source/Developer/Zen/Private/ZenServerInterface.cpp
--- a/Engine/Source/Developer/Zen/Private/ZenServerInterface.cpp
+++ b/Engine/Source/Developer/Zen/Private/ZenServerInterface.cpp
@@ -231,15 +231,32 @@ static bool GetZenVersion(const FString& UtilityPath, const FString& ServicePath
 		return false;
 	}
 
+	const uint64 StartTime = FPlatformTime::Cycles64();
 	while (MonitoredUtilityProcess.Update())
 	{
-		FPlatformProcess::Sleep(0.1f);
-		if (MonitoredUtilityProcess.GetDuration().GetTotalSeconds() > 10)
+		double Duration = FPlatformTime::ToSeconds64(FPlatformTime::Cycles64() - StartTime);
+		if (Duration > 10.0)
 		{
 			MonitoredUtilityProcess.Cancel(true);
-			UE_LOG(LogZenServiceInstance, Warning, TEXT("Cancelled launch of zen utility for gathering version data: '%s'."), *UtilityPath);
+			UE_LOG(LogZenServiceInstance, Warning, TEXT("Cancelled launch of zen utility for gathering version data: '%s' due to timeout."), *UtilityPath);
+
+			// Wait for execution to be terminated
+			while (MonitoredUtilityProcess.Update())
+			{
+				Duration = FPlatformTime::ToSeconds64(FPlatformTime::Cycles64() - StartTime);
+				if (Duration > 15.0)
+				{
+					UE_LOG(LogZenServiceInstance, Warning, TEXT("Cancelled launch of zen utility for gathering version data: '%s'. Failed waiting for termination."), *UtilityPath);
+					break;
+				}
+				FPlatformProcess::Sleep(0.2f);
+			}
+
+			FString OutputString = MonitoredUtilityProcess.GetFullOutputWithoutDelegate();
+			UE_LOG(LogZenServiceInstance, Warning, TEXT("Launch of zen utility for gathering version data: '%s' failed. Output: '%s'"), *UtilityPath, *OutputString);
 			return false;
 		}
+		FPlatformProcess::Sleep(0.1f);
 	}
 
 	FString OutputString = MonitoredUtilityProcess.GetFullOutputWithoutDelegate();
