Pull in XCurl instead of libCurl on Windows if necessary.

#rnx

diff --git a/Engine/Source/Developer/DevHttp/DevHttp.Build.cs b/Engine/Source/Developer/DevHttp/DevHttp.Build.cs
--- a/Engine/Source/Developer/DevHttp/DevHttp.Build.cs
+++ b/Engine/Source/Developer/DevHttp/DevHttp.Build.cs
@@ -9,7 +9,14 @@ public class DevHttp : ModuleRules
 	{
 		PrivateDependencyModuleNames.AddRange(new string[] { "Core", "SSL" });
 		
-		AddEngineThirdPartyPrivateStaticDependencies(Target, "libcurl");
+		if (Target.Platform.IsInGroup(UnrealPlatformGroup.Windows) && Target.WindowsPlatform.bUseXCurl)
+		{
+			PublicDependencyModuleNames.Add("XCurl");
+		}
+		else
+		{
+			AddEngineThirdPartyPrivateStaticDependencies(Target, "libcurl");
+		}
 		AddEngineThirdPartyPrivateStaticDependencies(Target, "OpenSSL");
 	}
 }
diff --git a/Engine/Source/Developer/Zen/Zen.Build.cs b/Engine/Source/Developer/Zen/Zen.Build.cs
--- a/Engine/Source/Developer/Zen/Zen.Build.cs
+++ b/Engine/Source/Developer/Zen/Zen.Build.cs
@@ -10,8 +10,15 @@ public class Zen : ModuleRules
 		PrivateDependencyModuleNames.AddRange(new string[] { "Core", "Sockets", "SSL", "Json" });
 		PrivateIncludePathModuleNames.Add("DesktopPlatform");
 		PrivateIncludePathModuleNames.Add("Analytics");
-		
-		AddEngineThirdPartyPrivateStaticDependencies(Target, "libcurl");
+
+		if (Target.Platform.IsInGroup(UnrealPlatformGroup.Windows) && Target.WindowsPlatform.bUseXCurl)
+		{
+			PublicDependencyModuleNames.Add("XCurl");
+		}
+		else
+		{		
+			AddEngineThirdPartyPrivateStaticDependencies(Target, "libcurl");
+		}
 		AddEngineThirdPartyPrivateStaticDependencies(Target, "OpenSSL");
 
 		CppCompileWarningSettings.UnsafeTypeCastWarningLevel = WarningLevel.Error;
