Add plugin support to zen launcher
#rnx

diff --git a/Engine/Source/Developer/Zen/Private/ZenServerInterface.cpp b/Engine/Source/Developer/Zen/Private/ZenServerInterface.cpp
--- a/Engine/Source/Developer/Zen/Private/ZenServerInterface.cpp
+++ b/Engine/Source/Developer/Zen/Private/ZenServerInterface.cpp
@@ -31,6 +31,7 @@
 #include "Misc/ScopeExit.h"
 #include "Serialization/MemoryReader.h"
 #include "Serialization/CompactBinarySerialization.h"
+#include "Serialization/CompactBinaryValue.h"
 #include "Serialization/JsonReader.h"
 #include "Serialization/JsonSerializer.h"
 #include "Serialization/JsonWriter.h"
@@ -164,6 +165,12 @@ GetServiceLinkPath()
 	return FPaths::Combine(GetLocalInstallPath(), TEXT("zen.link"));
 }
 
+static FString
+GetServicePluginsConfigPath()
+{
+	return FPaths::Combine(GetLocalInstallPath(), TEXT("zen_plugins.json"));
+}
+
 static FString
 GetServiceCopyInstallPath()
 {
@@ -786,6 +793,105 @@ static FString ZenGetInstallModeToString(FServiceAutoLaunchSettings::EInstallMod
 	return "";
 }
 
+bool FServicePluginSettings::ReadFromConfig(const FString& InPluginName)
+{
+	const FString PluginSectionName = FString::Format(TEXT("Zen.Plugin.{0}"), {*InPluginName}); 
+	const FConfigSection* PluginSection = GConfig->GetSection(*PluginSectionName, false, GEngineIni);
+	if (!PluginSection)
+	{
+		UE_LOG(LogZenServiceInstance, Warning, TEXT("Unable to find config section '%s'"), *PluginSectionName);
+		return false;
+	}
+
+	bool bHasName = false;
+	for (FConfigSection::TConstIterator It(*PluginSection); It; ++It)
+	{
+		if (It.Key() == TEXT("Name"))
+		{
+			Name = It.Value().GetValue();
+			// TODO load plugins from project dir
+			FString AbsPathRoot = FPaths::ConvertRelativePathToFull(FPaths::Combine(*FPaths::RootDir(), *Name));
+			if (IFileManager::Get().FileExists(*AbsPathRoot))
+			{
+				AbsPath = AbsPathRoot;
+			}
+			else
+			{
+				UE_LOG(LogZenServiceInstance, Warning, TEXT("Can't find file for plugin '%s', tried '%s'"), *Name, *AbsPathRoot);
+				return false;
+			}
+			bHasName = true;
+		}
+		else
+		{
+			Options.Add(It.Key(), It.Value().GetValue());
+		}
+	}
+
+	if (!bHasName)
+	{
+		UE_LOG(LogZenServiceInstance, Warning, TEXT("Settings for plugin defined in section '%s' are missing 'Name' value"), *PluginSectionName);
+	}
+
+	return bHasName;
+}
+
+bool FServicePluginSettings::ReadFromCompactBinary(FCbFieldView Field)
+{
+	if (!Field.IsObject())
+	{
+		return false;
+	}
+
+	bool bValid = true;
+
+	Name = FString(Field["Name"].AsString());
+	bValid &= !Name.IsEmpty();
+
+	AbsPath = FString(Field["AbsPath"].AsString());
+	bValid &= !AbsPath.IsEmpty();
+
+	if (FCbObjectView OptionsObject = Field["Options"].AsObjectView())
+	{
+		for (FCbFieldViewIterator It = OptionsObject.CreateViewIterator(); It;)
+		{
+			if (!It->IsString())
+			{
+				bValid = false;
+				break;
+			}
+			FString OptionName = FString(It->AsString());
+			++It;
+			if (!It->IsString())
+			{
+				bValid = false;
+				break;
+			}
+			FString OptionValue = FString(It->AsString());
+			Options.Add(FName(OptionName), OptionValue);
+		}
+	}
+
+	return bValid;
+}
+
+void FServicePluginSettings::WriteToCompactBinary(FCbWriter& Writer) const
+{
+	Writer.BeginObject();
+	Writer << "Name" << Name;
+	Writer << "AbsPath" << AbsPath;
+	if (!Options.IsEmpty())
+	{
+		Writer.BeginObject("Options");
+		for (const TPair<FName, FString>& Option: Options)
+		{
+			Writer << Option.Key << Option.Value;
+		}
+		Writer.EndObject();
+	}
+	Writer.EndObject();
+}
+
 bool
 FServiceSettings::ReadFromConfig()
 {
@@ -847,6 +953,17 @@ FServiceSettings::ReadFromConfig()
 			GConfig->GetString(AutoLaunchConfigSection, TEXT("ExtraArgs"), AutoLaunchSettings.ExtraArgs, GEngineIni);
 
 			ReadUInt16FromConfig(AutoLaunchConfigSection, TEXT("DesiredPort"), AutoLaunchSettings.DesiredPort, GEngineIni);
+			AutoLaunchSettings.Plugins.Empty();
+			TArray<FString> PluginNames;
+			GConfig->GetArray(AutoLaunchConfigSection, TEXT("Plugins"), PluginNames, GEngineIni);
+			for (FString PluginName : PluginNames)
+			{
+				FServicePluginSettings PluginSettings = {};
+				if (PluginSettings.ReadFromConfig(PluginName))
+				{
+					AutoLaunchSettings.Plugins.Add(PluginSettings);
+				}
+			}
 			GConfig->GetBool(AutoLaunchConfigSection, TEXT("ShowConsole"), AutoLaunchSettings.bShowConsole, GEngineIni);
 			GConfig->GetBool(AutoLaunchConfigSection, TEXT("LimitProcessLifetime"), AutoLaunchSettings.bLimitProcessLifetime, GEngineIni);
 			ApplyProcessLifetimeOverride(AutoLaunchSettings.bLimitProcessLifetime);
@@ -890,6 +1007,18 @@ FServiceSettings::ReadFromCompactBinary(FCbFieldView Field)
 					AutoLaunchSettings.DataPath = FString(AutoLaunchSettingsObject["DataPath"].AsString());
 					AutoLaunchSettings.ExtraArgs = FString(AutoLaunchSettingsObject["ExtraArgs"].AsString());
 					AutoLaunchSettings.DesiredPort = AutoLaunchSettingsObject["DesiredPort"].AsInt16();
+					if (FCbArrayView PluginsView = AutoLaunchSettingsObject["Plugins"].AsArrayView())
+					{
+						AutoLaunchSettings.Plugins.Empty();
+						for (FCbFieldView& PluginView : PluginsView)
+						{
+							FServicePluginSettings PluginSettings = {};
+							if (PluginSettings.ReadFromCompactBinary(PluginView))
+							{
+								AutoLaunchSettings.Plugins.Add(PluginSettings);
+							}
+						}
+					}
 					AutoLaunchSettings.bShowConsole = AutoLaunchSettingsObject["ShowConsole"].AsBool();
 					AutoLaunchSettings.bIsDefaultDataPath = AutoLaunchSettingsObject["IsDefaultDataPath"].AsBool();
 					AutoLaunchSettings.bLimitProcessLifetime = AutoLaunchSettingsObject["LimitProcessLifetime"].AsBool();
@@ -955,6 +1084,16 @@ FServiceSettings::WriteToCompactBinary(FCbWriter& Writer) const
 		Writer << "DataPath" << AutoLaunchSettings.DataPath;
 		Writer << "ExtraArgs" <<AutoLaunchSettings.ExtraArgs;
 		Writer << "DesiredPort" << AutoLaunchSettings.DesiredPort;
+
+		if (!AutoLaunchSettings.Plugins.IsEmpty())
+		{
+			Writer.BeginArray("Plugins");
+			for (const FServicePluginSettings& PluginSettings : AutoLaunchSettings.Plugins)
+			{
+				PluginSettings.WriteToCompactBinary(Writer);
+			}
+			Writer.EndArray();
+		}
 		Writer << "ShowConsole" << AutoLaunchSettings.bShowConsole;
 		Writer << "IsDefaultDataPath" << AutoLaunchSettings.bIsDefaultDataPath;
 		Writer << "LimitProcessLifetime" << AutoLaunchSettings.bLimitProcessLifetime;
@@ -1243,6 +1382,13 @@ DetermineCmdLineWithoutTransientComponents(const FServiceAutoLaunchSettings& InS
 		Parms.Append(InSettings.ExtraArgs);
 	}
 
+	if (!InSettings.Plugins.IsEmpty())
+	{
+		Parms.AppendChar(TEXT(' '));
+		Parms.Appendf(TEXT("--plugins-config \"%s\""),
+			*GetServicePluginsConfigPath());
+	}
+
 	FString LogCommandLineOverrideValue;
 	if (FParse::Value(FCommandLine::Get(), TEXT("ZenLogPath="), LogCommandLineOverrideValue))
 	{
@@ -2024,6 +2170,17 @@ PromptUserOfFailedShutDownOfExistingProcess(uint16 Port)
 bool
 FZenServiceInstance::ConditionalUpdateLocalInstall(FServiceAutoLaunchSettings::EInstallMode InstallMode)
 {
+	// Normally plugins config is written only if we update zen local install,
+	// extra command line provided to force plugins update.
+	const bool bForceZenPluginsInstall = FParse::Param(FCommandLine::Get(), TEXT("ForceZenPluginsInstall"));
+	// If current settings require plugins, we will pass "--plugins-config" command line argument to zen server,
+	// at that point plugins config file should exist, so create one if it doesn't.
+	const bool bMissingZenPluginsConfig = Settings.IsRequirePlugins() && !IFileManager::Get().FileExists(*GetServicePluginsConfigPath());
+	if (bForceZenPluginsInstall || bMissingZenPluginsConfig)
+	{
+		WriteLocalPluginsConfig();
+	}
+
 	if (InstallMode == FServiceAutoLaunchSettings::EInstallMode::Link)
 	{
 		const FString LinkPath = GetServiceLinkPath();
@@ -2114,6 +2271,8 @@ FZenServiceInstance::ConditionalUpdateLocalInstall(FServiceAutoLaunchSettings::E
 			return false;
 		}
 
+		WriteLocalPluginsConfig();
+
 		FString ServiceCopyInstallPath = GetServiceCopyInstallPath();
 		if (FileManager.FileExists(*ServiceCopyInstallPath))
 		{
@@ -2172,6 +2331,8 @@ FZenServiceInstance::ConditionalUpdateLocalInstall(FServiceAutoLaunchSettings::E
 
 			AttemptFileCopyWithRetries(*InstallVersionCache, *InTreeVersionCache, 1.0);
 
+			WriteLocalPluginsConfig();
+
 			bMainExecutablesUpdated = true;
 		}
 
@@ -2225,6 +2386,45 @@ FZenServiceInstance::ConditionalUpdateLocalInstall(FServiceAutoLaunchSettings::E
 	}
 }
 
+bool
+FZenServiceInstance::WriteLocalPluginsConfig()
+{
+	if (!Settings.IsRequirePlugins())
+	{
+		return false;
+	}
+
+	TConstArrayView<FServicePluginSettings> PluginsSettings = Settings.SettingsVariant.Get<FServiceAutoLaunchSettings>().Plugins;
+
+	TArray<TSharedPtr<FJsonValue>> JsonArray;
+
+	// replace this with UStructToJsonObject or similar in the future
+	for (const FServicePluginSettings& PluginSettings : PluginsSettings)
+	{
+		TSharedPtr<FJsonObject> JsonObject = MakeShareable(new FJsonObject);
+
+		JsonObject->SetStringField(TEXT("name"), PluginSettings.AbsPath); // zen server can treat name as absolute path
+		for (const TPair<FName, FString>& Option: PluginSettings.Options)
+		{
+			JsonObject->SetStringField(Option.Key.ToString(), Option.Value);
+		}
+
+		TSharedPtr<FJsonValueObject> JsonValueObject = MakeShareable(new FJsonValueObject(MoveTemp(JsonObject)));
+		JsonArray.Add(JsonValueObject);
+	}
+
+	FString OutputString;
+	TSharedRef<TJsonWriter<>> Writer = TJsonWriterFactory<>::Create(&OutputString);
+	if (!FJsonSerializer::Serialize(JsonArray, Writer))
+	{
+		return false;
+	}
+
+	const FString OutputPath = GetServicePluginsConfigPath();
+	UE_LOG(LogZenServiceInstance, Display, TEXT("Writing plugin configuration to '%s'"), *OutputPath);
+	return FFileHelper::SaveStringToFile(OutputString, *OutputPath);
+}
+
 bool
 FZenServiceInstance::AutoLaunch(const FServiceAutoLaunchSettings& InSettings, const FString& ExecutablePath, const FString& ExecutionContextFilePath, FString& OutHostName, uint16& OutPort)
 {
diff --git a/Engine/Source/Developer/Zen/Public/Experimental/ZenServerInterface.h b/Engine/Source/Developer/Zen/Public/Experimental/ZenServerInterface.h
--- a/Engine/Source/Developer/Zen/Public/Experimental/ZenServerInterface.h
+++ b/Engine/Source/Developer/Zen/Public/Experimental/ZenServerInterface.h
@@ -36,11 +36,23 @@ struct FServiceConnectSettings
 	uint16 Port = 8558;
 };
 
+struct FServicePluginSettings
+{
+	FString Name;
+	FString AbsPath;
+	TMap<FName, FString> Options;
+
+	UE_API bool ReadFromConfig(const FString& InPluginName);
+	UE_API bool ReadFromCompactBinary(FCbFieldView Field);
+	UE_API void WriteToCompactBinary(FCbWriter& Writer) const;
+};
+
 struct FServiceAutoLaunchSettings
 {
 	FString DataPath;
 	FString ExtraArgs;
 	uint16 DesiredPort = 8558;
+	TArray<FServicePluginSettings> Plugins;
 	bool bShowConsole = false;
 	bool bIsDefaultDataPath = false;
 	bool bLimitProcessLifetime = false;
@@ -60,6 +72,7 @@ struct FServiceSettings
 
 	inline bool IsAutoLaunch() const { return SettingsVariant.IsType<FServiceAutoLaunchSettings>(); }
 	inline bool IsConnectExisting() const { return SettingsVariant.IsType<FServiceConnectSettings>(); }
+	inline bool IsRequirePlugins() const { return IsAutoLaunch() && !SettingsVariant.Get<FServiceAutoLaunchSettings>().Plugins.IsEmpty(); }
 
 	UE_API bool ReadFromConfig();
 	UE_API bool ReadFromCompactBinary(FCbFieldView Field);
@@ -279,6 +292,7 @@ private:
 
 	void Initialize();
 	bool ConditionalUpdateLocalInstall(FServiceAutoLaunchSettings::EInstallMode InstallMode);
+	bool WriteLocalPluginsConfig();
 	static bool AutoLaunch(const FServiceAutoLaunchSettings& InSettings, const FString& ExecutablePath, const FString& ExecutionContextFilePath, FString& OutHostName, uint16& OutPort);
 	bool PingService(uint32_t AttemptCount);
 
