ZenServer over a Unix socket.

#rnx

diff --git a/Engine/Source/Developer/Android/AndroidDeviceDetection/Private/AndroidDeviceDetectionModule.cpp b/Engine/Source/Developer/Android/AndroidDeviceDetection/Private/AndroidDeviceDetectionModule.cpp
--- a/Engine/Source/Developer/Android/AndroidDeviceDetection/Private/AndroidDeviceDetectionModule.cpp
+++ b/Engine/Source/Developer/Android/AndroidDeviceDetection/Private/AndroidDeviceDetectionModule.cpp
@@ -679,7 +679,7 @@ private:
 						41899,	// Network file server, DEFAULT_TCP_FILE_SERVING_PORT in NetworkMessage.h
 						1981,	// Unreal Insights data collection, TraceInsightsModule.cpp
 #if UE_WITH_ZEN
-						UE::Zen::IsDefaultServicePresent() ? UE::Zen::GetDefaultServiceInstance().GetPort() : uint16(0), // Zen Store, usually defaults to 8558
+						UE::Zen::IsDefaultServicePresent() ? UE::Zen::GetDefaultServiceInstance().GetEndpoint().GetPort() : uint16(0), // Zen Store, usually defaults to 8558
 #endif
 						0 // end of list
 					};
diff --git a/Engine/Source/Developer/DerivedDataCache/Private/Tests/HttpCacheStoreTest.cpp b/Engine/Source/Developer/DerivedDataCache/Private/Tests/HttpCacheStoreTest.cpp
--- a/Engine/Source/Developer/DerivedDataCache/Private/Tests/HttpCacheStoreTest.cpp
+++ b/Engine/Source/Developer/DerivedDataCache/Private/Tests/HttpCacheStoreTest.cpp
@@ -650,7 +650,7 @@ bool FHttpCacheStoreTest::RunTest(const FString& Parameters)
 	FServiceAutoLaunchSettings& ZenTestAutoLaunchSettings = ZenTestServiceSettings.SettingsVariant.Get<FServiceAutoLaunchSettings>();
 	ZenTestAutoLaunchSettings.DataPath = FPaths::ConvertRelativePathToFull(FPaths::Combine(FPaths::EngineSavedDir(), "ZenUnitTest"));
 	ZenTestAutoLaunchSettings.ExtraArgs = FString::Printf(TEXT("--http asio --upstream-zen-url \"http://localhost:%d\""),
-		ScopeZenUpstreamService.GetInstance().GetPort()
+		ScopeZenUpstreamService.GetInstance().GetEndpoint().GetPort()
 	);
 	ZenTestAutoLaunchSettings.DesiredPort = 13337; // Avoid the normal default port
 	ZenTestAutoLaunchSettings.bShowConsole = true;
@@ -661,7 +661,7 @@ bool FHttpCacheStoreTest::RunTest(const FString& Parameters)
 	FServiceAutoLaunchSettings& ZenTestSiblingAutoLaunchSettings = ZenTestServiceSiblingSettings.SettingsVariant.Get<FServiceAutoLaunchSettings>();
 	ZenTestSiblingAutoLaunchSettings.DataPath = FPaths::ConvertRelativePathToFull(FPaths::Combine(FPaths::EngineSavedDir(), "ZenUnitTestSibling"));
 	ZenTestSiblingAutoLaunchSettings.ExtraArgs = FString::Printf(TEXT("--http asio --upstream-zen-url \"http://localhost:%d\""),
-		ScopeZenUpstreamSiblingService.GetInstance().GetPort()
+		ScopeZenUpstreamSiblingService.GetInstance().GetEndpoint().GetPort()
 	);
 	ZenTestSiblingAutoLaunchSettings.DesiredPort = 13338; // Avoid the normal default port
 	ZenTestSiblingAutoLaunchSettings.bShowConsole = true;
@@ -669,11 +669,11 @@ bool FHttpCacheStoreTest::RunTest(const FString& Parameters)
 
 	FScopeZenService ScopeZenSiblingService(MoveTemp(ZenTestServiceSiblingSettings));
 	TUniquePtr<ILegacyCacheStore> ZenIntermediarySiblingBackend(CreateZenCacheStore(TEXT("TestSibling"),
-		*FString::Printf(TEXT("Host=%s, StructuredNamespace=%s"), ScopeZenSiblingService.GetInstance().GetURL(), *TestNamespace), nullptr));
+		*FString::Printf(TEXT("Host=%s, StructuredNamespace=%s"), *ScopeZenSiblingService.GetInstance().GetEndpoint().GetURL(), *TestNamespace), nullptr));
 
 	FScopeZenService ScopeZenService(MoveTemp(ZenTestServiceSettings));
 	TUniquePtr<ILegacyCacheStore> ZenIntermediaryBackend(CreateZenCacheStore(TEXT("Test"),
-		*FString::Printf(TEXT("Host=%s, StructuredNamespace=%s"), ScopeZenService.GetInstance().GetURL(), *TestNamespace), nullptr));
+		*FString::Printf(TEXT("Host=%s, StructuredNamespace=%s"), *ScopeZenService.GetInstance().GetEndpoint().GetURL(), *TestNamespace), nullptr));
 	auto WaitForZenPushToUpstream = [](ILegacyCacheStore* ZenBackend, TConstArrayView<FCacheRecord> Records)
 	{
 		// TODO: Expecting a legitimate means to wait for zen to finish pushing records to its upstream in the future
diff --git a/Engine/Source/Developer/DerivedDataCache/Private/ZenCacheStore.cpp b/Engine/Source/Developer/DerivedDataCache/Private/ZenCacheStore.cpp
--- a/Engine/Source/Developer/DerivedDataCache/Private/ZenCacheStore.cpp
+++ b/Engine/Source/Developer/DerivedDataCache/Private/ZenCacheStore.cpp
@@ -213,6 +213,8 @@ private:
 		FValue& OutValue,
 		FRequestStats& Stats) const;
 
+	void SetRequestUri(IHttpRequest& Request, FAnsiStringView Path);
+
 	class FPutOp;
 	class FGetOp;
 	class FPutValueOp;
@@ -247,7 +249,7 @@ private:
 	int32 CacheRecordBatchSize = 8;
 	int32 CacheChunksBatchSize = 8;
 	FBackendDebugOptions DebugOptions;
-	TAnsiStringBuilder<256> RpcUri;
+	FAnsiStringView RpcUri = ANSITEXTVIEW("/z$/$rpc");
 	std::atomic<int64> LastPerformanceEvaluationTicks;
 	TOptional<FThread> PerformanceEvaluationThread;
 	FManualResetEvent PerformanceEvaluationThreadShutdownEvent;
@@ -303,6 +305,29 @@ void FZenCacheStore::StartRequests(
 	}
 }
 
+void FZenCacheStore::SetRequestUri(IHttpRequest& Request, FAnsiStringView Path)
+{
+	using namespace Zen;
+
+	check(Path.Len() && Path[0] == '/');
+
+	const FZenServiceEndpoint& Endpoint = ZenService.GetInstance().GetEndpoint();
+
+	TAnsiStringBuilder<128> Uri;
+	Uri << Endpoint.GetURL();
+	Uri << Path;
+	Request.SetUri(Uri);
+
+	if (Endpoint.GetSocketType() != FZenServiceEndpoint::ESocketType::Unix)
+	{
+		return;
+	}
+
+	FStringView UnixSocketPath = Endpoint.GetName();
+	const auto& AsAnsi = StringCast<ANSICHAR>(UnixSocketPath.GetData(), UnixSocketPath.Len());
+	Request.SetUnixSocketPath(AsAnsi);
+}
+
 bool FZenCacheStore::ValidateAndPrepareReceivedValue(
 		FStringView Name,
 		const FCacheKey& Key,
@@ -1694,8 +1719,6 @@ void FZenCacheStore::Initialize(const FZenCacheStoreParams& Params)
 	LastStorageSizeUpdateTicks.store(FDateTime::UtcNow().GetTicks(), std::memory_order_relaxed);
 	Namespace = Params.Namespace;
 
-	RpcUri << ZenService.GetInstance().GetURL() << ANSITEXTVIEW("/z$/$rpc");
-
 	const uint32 MaxConnections = uint32(FMath::Clamp(FPlatformMisc::NumberOfCoresIncludingHyperthreads(), 8, 64));
 	constexpr uint32 RequestPoolSize = 128;
 	constexpr uint32 RequestPoolOverflowSize = 128;
@@ -1717,9 +1740,7 @@ void FZenCacheStore::Initialize(const FZenCacheStoreParams& Params)
 		ReadinessClientParams.bBypassProxy = Params.bBypassProxy;
 		THttpUniquePtr<IHttpClient> ReadinessClient = ConnectionPool->CreateClient(ReadinessClientParams);
 		THttpUniquePtr<IHttpRequest> ReadinessRequest = ReadinessClient->TryCreateRequest({});
-		TAnsiStringBuilder<256> StatusUri;
-		StatusUri << ZenService.GetInstance().GetURL() << ANSITEXTVIEW("/health/ready");
-		ReadinessRequest->SetUri(StatusUri);
+		SetRequestUri(*ReadinessRequest, ANSITEXTVIEW("/health/ready"));
 		ReadinessRequest->SetMethod(EHttpMethod::Get);
 		ReadinessRequest->AddAcceptType(EHttpMediaType::Text);
 		EHealth Health = EHealth::Unknown;
@@ -1729,12 +1750,14 @@ void FZenCacheStore::Initialize(const FZenCacheStoreParams& Params)
 		ReadinessRequest->Send(&HealthReceiver, ReadinessResponse);
 		bReady = Health == EHealth::Ok; // -V547
 
+		const Zen::FZenServiceEndpoint& ZenServiceEndpoint = ZenService.GetInstance().GetEndpoint();
+
 		if (ReadinessResponse->GetErrorCode() == EHttpErrorCode::None &&
 			(ReadinessResponse->GetStatusCode() >= 200 && ReadinessResponse->GetStatusCode() <= 299))
 		{
 			UE_LOG(LogDerivedDataCache, Display,
 				TEXT("%s: Using ZenServer HTTP service at %s with namespace %s status: %s."),
-				*GetName(), ZenService.GetInstance().GetURL(), *Namespace, *ResponseString);
+				*GetName(), ZenServiceEndpoint.GetName().GetData(), *Namespace, *ResponseString);
 		}
 		else
 		{
@@ -1742,14 +1765,14 @@ void FZenCacheStore::Initialize(const FZenCacheStoreParams& Params)
 			{
 				UE_LOG(LogDerivedDataCache, Warning,
 					TEXT("%s: Unable to reach ZenServer HTTP service at %s with namespace %s. Status: %d . Response: %s"),
-					*GetName(), ZenService.GetInstance().GetURL(), *Namespace, ReadinessResponse->GetStatusCode(),
+					*GetName(), ZenServiceEndpoint.GetName().GetData(), *Namespace, ReadinessResponse->GetStatusCode(),
 					*ResponseString);
 			}
 			else
 			{
 				UE_LOG(LogDerivedDataCache, Display,
 					TEXT("%s: Unable to reach ZenServer HTTP service at %s with namespace %s. Status: %d . Response: %s"),
-					*GetName(), ZenService.GetInstance().GetURL(), *Namespace, ReadinessResponse->GetStatusCode(),
+					*GetName(), ZenServiceEndpoint.GetName().GetData(), *Namespace, ReadinessResponse->GetStatusCode(),
 					*ResponseString);
 			}
 		}
@@ -1852,7 +1875,7 @@ void FZenCacheStore::CreateRpcRequest(IRequestOwner& Owner, FHttpRequestQueue::F
 	{
 		if (Request)
 		{
-			Request->SetUri(RpcUri);
+			SetRequestUri(*Request, RpcUri);
 			Request->SetMethod(EHttpMethod::Post);
 			Request->AddAcceptType(EHttpMediaType::CbPackage);
 		}
@@ -1918,9 +1941,7 @@ void FZenCacheStore::ActivatePerformanceEvaluationThread()
 					{
 						return;
 					}
-					TAnsiStringBuilder<256> StatusUri;
-					StatusUri << ZenService.GetInstance().GetURL() << ANSITEXTVIEW("/health/ready");
-					Request->SetUri(StatusUri);
+					SetRequestUri(*Request, ANSITEXTVIEW("/health/ready"));
 					Request->SetMethod(EHttpMethod::Get);
 					Request->AddAcceptType(EHttpMediaType::Text);
 					new FAsyncHealthReceiver(MoveTemp(Request), &Owner, ZenService.GetInstance(), [this, StartTime = FMonotonicTimePoint::Now()](const THttpUniquePtr<IHttpResponse>& HttpResponse, EHealth Health)
diff --git a/Engine/Source/Developer/IoStoreUtilities/Internal/ZenStoreHttpClient.h b/Engine/Source/Developer/IoStoreUtilities/Internal/ZenStoreHttpClient.h
--- a/Engine/Source/Developer/IoStoreUtilities/Internal/ZenStoreHttpClient.h
+++ b/Engine/Source/Developer/IoStoreUtilities/Internal/ZenStoreHttpClient.h
@@ -22,6 +22,7 @@ class FCbPackage;
 class FCbObject;
 
 namespace UE {
+	class IHttpRequest;
 namespace Zen {
 	struct FZenHttpRequestPool;
 	class FZenHttpRequest;
@@ -71,13 +72,8 @@ public:
 	UE_API TIoStatusOr<FIoBuffer> ReadOpLogAttachment(FStringView Id);
 
 #if UE_WITH_ZEN
-	const TCHAR* GetHostName() const { return ZenService.GetInstance().GetHostName(); }
-	uint16 GetPort() const { return ZenService.GetInstance().GetPort(); }
 	const UE::Zen::FZenServiceInstance& GetZenServiceInstance() const { return ZenService.GetInstance(); }
 	UE::Zen::FZenServiceInstance& GetZenServiceInstance() { return ZenService.GetInstance(); }
-#else // Default to localhost:8558 for platforms where Zen wouldn't be supported yet
-	const TCHAR* GetHostName() const { return TEXT("localhost"); }
-	uint16 GetPort() const { return 8558; }
 #endif
 
 	UE_API TFuture<TIoStatusOr<FCbObject>> GetProjectInfo();
@@ -97,6 +93,7 @@ private:
 	UE_API void RpcAsync(FAnsiStringView Uri, FCbObject RpcRequestObject, TUniqueFunction<void(const FCbPackage& Response)>&& ResponseCallback);
 	UE_API bool Delete(Zen::FZenHttpRequest& Request, FStringView Uri);
 	UE_API bool ShouldRecoverAndRetry(Zen::FZenHttpRequest& Request);
+	void SetRequestUri(class IHttpRequest& Request, FAnsiStringView Path);
 
 	template<bool bAsync, typename T>
 	void ReadChunks_Internal(TConstArrayView<T> Requests, TUniqueFunction<void(const T& Id, TIoStatusOr<FIoBuffer> Result)>&& ChunkCallback);
diff --git a/Engine/Source/Developer/IoStoreUtilities/Private/ZenStoreHttpClient.cpp b/Engine/Source/Developer/IoStoreUtilities/Private/ZenStoreHttpClient.cpp
--- a/Engine/Source/Developer/IoStoreUtilities/Private/ZenStoreHttpClient.cpp
+++ b/Engine/Source/Developer/IoStoreUtilities/Private/ZenStoreHttpClient.cpp
@@ -76,7 +76,7 @@ FZenStoreHttpClient::~FZenStoreHttpClient()
 void
 FZenStoreHttpClient::InitCommon()
 {
-	RequestPool = MakeUnique<Zen::FZenHttpRequestPool>(ZenService.GetInstance().GetURL(), PoolEntryCount);
+	RequestPool = MakeUnique<Zen::FZenHttpRequestPool>(ZenService.GetInstance().GetEndpoint(), PoolEntryCount);
 	ClientState = MakeUnique<Zen::Internal::FZenHttpClientState>();
 
 	const uint32 MaxConnections = uint32(FMath::Clamp(FPlatformMisc::NumberOfCoresIncludingHyperthreads(), 8, 64));
@@ -167,7 +167,7 @@ FZenStoreHttpClient::TryCreateProject(FStringView InProjectId,
 
 	ProjectPath = WriteToString<128>("/prj/", InProjectId);
 	OplogPath = WriteToString<128>("/prj/", InProjectId, "/oplog/", InOplogId);
-	OplogRpcPath = WriteToAnsiString<128>(ZenService.GetInstance().GetURL(), "/prj/", InProjectId, "/oplog/", InOplogId, "/rpc");
+	OplogRpcPath = WriteToAnsiString<128>("/prj/", InProjectId, "/oplog/", InOplogId, "/rpc");
 	OplogNewEntryPath = WriteToString<128>("/prj/", InProjectId, "/oplog/", InOplogId, "/new");
 	OplogPrepNewEntryPath = WriteToString<128>("/prj/", InProjectId, "/oplog/", InOplogId, "/prep");
 
@@ -327,6 +327,29 @@ bool FZenStoreHttpClient::Post(Zen::FZenHttpRequest& Request, FStringView Uri, F
 	return Res == Zen::FZenHttpRequest::Result::Success;
 }
 
+void FZenStoreHttpClient::SetRequestUri(IHttpRequest& Request, FAnsiStringView Path)
+{
+	using namespace Zen;
+
+	check(Path.Len() && Path[0] == '/');
+
+	const FZenServiceEndpoint& Endpoint = ZenService.GetInstance().GetEndpoint();
+
+	TAnsiStringBuilder<128> Uri;
+	Uri << Endpoint.GetURL();
+	Uri << Path;
+	Request.SetUri(Uri);
+
+	if (Endpoint.GetSocketType() != FZenServiceEndpoint::ESocketType::Unix)
+	{
+		return;
+	}
+
+	FStringView UnixSocketPath = Endpoint.GetName();
+	const auto& AsAnsi = StringCast<ANSICHAR>(UnixSocketPath.GetData(), UnixSocketPath.Len());
+	Request.SetUnixSocketPath(AsAnsi);
+}
+
 void FZenStoreHttpClient::Rpc(FAnsiStringView Uri, FCbObject RpcRequestObject, FCbPackage& OutResponse)
 {
 	THttpUniquePtr<IHttpRequest> Request = ClientState->RequestQueue->CreateRequest({});
@@ -338,7 +361,7 @@ void FZenStoreHttpClient::Rpc(FAnsiStringView Uri, FCbObject RpcRequestObject, F
 	constexpr int MaxAttempts = 2;
 	int Attempt = 0;
 	bool bRetry = false;
-	Request->SetUri(Uri);
+	SetRequestUri(*Request, Uri);
 	Request->SetMethod(EHttpMethod::Post);
 	Request->AddAcceptType(EHttpMediaType::CbPackage);
 	Request->SetContentType(EHttpMediaType::CbObject);
@@ -365,7 +388,7 @@ void FZenStoreHttpClient::RpcAsync(FAnsiStringView Uri, FCbObject RpcRequestObje
 		return;
 	}
 
-	Request->SetUri(Uri);
+	SetRequestUri(*Request, Uri);
 	Request->SetMethod(EHttpMethod::Post);
 	Request->AddAcceptType(EHttpMediaType::CbPackage);
 	Request->SetContentType(EHttpMediaType::CbObject);
diff --git a/Engine/Source/Developer/IoStoreUtilities/Private/ZenStoreWriter.cpp b/Engine/Source/Developer/IoStoreUtilities/Private/ZenStoreWriter.cpp
--- a/Engine/Source/Developer/IoStoreUtilities/Private/ZenStoreWriter.cpp
+++ b/Engine/Source/Developer/IoStoreUtilities/Private/ZenStoreWriter.cpp
@@ -486,15 +486,21 @@ void FZenStoreWriter::Initialize(const FCookInfo& Info)
 
 		if (bOplogEstablished && OplogMarker)
 		{
+			FStringView ServiceHostName = TEXT("localhost");
+			uint16 ServicePort = 8558;
 			bool IsRunningLocally = false;
+
 #if UE_WITH_ZEN
-			IsRunningLocally = HttpClient->GetZenServiceInstance().IsServiceRunningLocally();
+			const Zen::FZenServiceInstance& ZenServiceInstance = HttpClient->GetZenServiceInstance();
+			IsRunningLocally = ZenServiceInstance.IsServiceRunningLocally();
+			ServiceHostName = ZenServiceInstance.GetEndpoint().GetHostName();
+			ServicePort = ZenServiceInstance.GetEndpoint().GetPort();
 #endif
 			TSharedRef<TJsonWriter<UTF8CHAR, TPrettyJsonPrintPolicy<UTF8CHAR>>> Writer = TJsonWriterFactory<UTF8CHAR, TPrettyJsonPrintPolicy<UTF8CHAR>>::Create(OplogMarker.Get());
 			Writer->WriteObjectStart();
 			Writer->WriteObjectStart(TEXT("zenserver"));
 			Writer->WriteValue(TEXT("islocalhost"), IsRunningLocally);
-			Writer->WriteValue(TEXT("hostname"), HttpClient->GetHostName());
+			Writer->WriteValue(TEXT("hostname"), ServiceHostName);
 			if (IsRunningLocally)
 			{
 				ISocketSubsystem* SocketSubsystem = ISocketSubsystem::Get();
@@ -545,7 +551,7 @@ void FZenStoreWriter::Initialize(const FCookInfo& Info)
 					}
 				}
 			}
-			Writer->WriteValue(TEXT("hostport"), HttpClient->GetPort());
+			Writer->WriteValue(TEXT("hostport"), ServicePort);
 			Writer->WriteValue(TEXT("projectid"), ProjectId);
 			Writer->WriteValue(TEXT("oplogid"), OplogId);
 			Writer->WriteObjectEnd();
@@ -729,6 +735,8 @@ void FZenStoreWriter::EndCook(const FCookInfo& Info)
 
 FZenStoreWriter::ZenHostInfo FZenStoreWriter::GetHostInfo() const
 {
+	const Zen::FZenServiceInstance& Instance = HttpClient->GetZenServiceInstance();
+
 	FZenStoreWriter::ZenHostInfo Info;
 	Info.ProjectId = ProjectId;
 	Info.OplogId = OplogId;
@@ -738,9 +746,9 @@ FZenStoreWriter::ZenHostInfo FZenStoreWriter::GetHostInfo() const
 	}
 	else
 	{
-		Info.HostName = HttpClient->GetHostName();
+		Info.HostName = Instance.GetEndpoint().GetHostName();
 	}
-	Info.HostPort = HttpClient->GetPort();
+	Info.HostPort = Instance.GetEndpoint().GetPort();
 	return Info;
 }
 
diff --git a/Engine/Source/Developer/Zen/Internal/ZenServerHttp.h b/Engine/Source/Developer/Zen/Internal/ZenServerHttp.h
--- a/Engine/Source/Developer/Zen/Internal/ZenServerHttp.h
+++ b/Engine/Source/Developer/Zen/Internal/ZenServerHttp.h
@@ -32,6 +32,8 @@ class FCbPackage;
 
 namespace UE::Zen {
 
+class FZenServiceEndpoint;
+
 static bool IsSuccessCode(int ResponseCode)
 {
 	return 200 <= ResponseCode && ResponseCode < 300;
@@ -86,7 +88,7 @@ inline FStringView GetMimeType(EContentType Type)
 class FZenHttpRequest
 {
 public:
-	ZEN_API FZenHttpRequest(FStringView InDomain, bool bInLogErrors, uint32_t ConnectTimeoutMS = 0);
+	ZEN_API FZenHttpRequest(const FZenServiceEndpoint& InDomain, bool bInLogErrors, uint32 ConnectTimeoutMS = 0);
 	ZEN_API ~FZenHttpRequest();
 
 	/**
@@ -221,6 +223,7 @@ private:
 	TArray64<uint8>			ResponseBuffer;	// If no other response buffer is set, this is where the response payload goes
 	TArray<FString>			Headers;
 	FString					Domain;
+	const FZenServiceEndpoint& Endpoint;
 
 	void AddHeader(FStringView Header, FStringView Value);
 
@@ -262,7 +265,7 @@ private:
   */
 struct FZenHttpRequestPool
 {
-	ZEN_API explicit FZenHttpRequestPool(FStringView InServiceUrl, uint32 PoolEntryCount = 16);
+	ZEN_API explicit FZenHttpRequestPool(const FZenServiceEndpoint& Endpoint, uint32 PoolEntryCount = 16);
 	ZEN_API ~FZenHttpRequestPool();
 
 	/** Block until a request is free. Once a request has been returned it is
diff --git a/Engine/Source/Developer/Zen/Private/ZenServerHttp.cpp b/Engine/Source/Developer/Zen/Private/ZenServerHttp.cpp
--- a/Engine/Source/Developer/Zen/Private/ZenServerHttp.cpp
+++ b/Engine/Source/Developer/Zen/Private/ZenServerHttp.cpp
@@ -60,9 +60,10 @@ namespace UE::Zen {
 #endif //PLATFORM_WINDOWS
 	};
 
-	FZenHttpRequest::FZenHttpRequest(FStringView InDomain, bool bInLogErrors, uint32_t ConnectTimeoutMS)
+	FZenHttpRequest::FZenHttpRequest(const FZenServiceEndpoint& InEndpoint, bool bInLogErrors, uint32 ConnectTimeoutMS)
 		: bLogErrors(bInLogErrors)
-		, Domain(InDomain)
+		, Domain(InEndpoint.GetURL())
+		, Endpoint(InEndpoint)
 	{
 		Curl = curl_easy_init();
 		Reset();
@@ -118,6 +119,13 @@ namespace UE::Zen {
 		curl_easy_setopt(Curl, CURLOPT_DEBUGFUNCTION, &FZenHttpRequest::FStatics::StaticDebugCallback);
 		curl_easy_setopt(Curl, CURLOPT_VERBOSE, 1L);
 #endif
+
+		if (Endpoint.GetSocketType() == FZenServiceEndpoint::ESocketType::Unix)
+		{
+			FStringView UnixSocketPath = Endpoint.GetName();
+			const auto& AsUtf8 = StringCast<UTF8CHAR>(UnixSocketPath.GetData(), UnixSocketPath.Len());
+			curl_easy_setopt(Curl, CURLOPT_UNIX_SOCKET_PATH, AsUtf8.Get());
+		}
 	}
 
 	void FZenHttpRequest::Initialize(bool bInLogErrors)
@@ -744,20 +752,14 @@ namespace UE::Zen {
 
 	//////////////////////////////////////////////////////////////////////////
 
-	FZenHttpRequestPool::FZenHttpRequestPool(FStringView InServiceUrl, uint32 PoolEntryCount)
+	FZenHttpRequestPool::FZenHttpRequestPool(const FZenServiceEndpoint& Endpoint, uint32 PoolEntryCount)
 	{
-		int32 LastSlash = 0;
-		while (InServiceUrl.FindLastChar(TEXT('/'), /* out */ LastSlash) && (LastSlash == (InServiceUrl.Len() - 1)))
-		{
-			InServiceUrl.LeftChopInline(1);
-		}
-
 		Pool.SetNum(PoolEntryCount);
 
 		for (uint8 i = 0; i < Pool.Num(); ++i)
 		{
 			Pool[i].IsAllocated = 0u;
-			Pool[i].Request = new FZenHttpRequest(InServiceUrl, true /* bLogErrors */);
+			Pool[i].Request = new FZenHttpRequest(Endpoint, true /* bLogErrors */);
 		}
 	}
 
diff --git a/Engine/Source/Developer/Zen/Private/ZenServerInterface.cpp b/Engine/Source/Developer/Zen/Private/ZenServerInterface.cpp
--- a/Engine/Source/Developer/Zen/Private/ZenServerInterface.cpp
+++ b/Engine/Source/Developer/Zen/Private/ZenServerInterface.cpp
@@ -51,6 +51,34 @@ namespace UE::Zen
 
 DEFINE_LOG_CATEGORY_STATIC(LogZenServiceInstance, Log, All);
 
+FZenServiceEndpoint::FZenServiceEndpoint(FStringView InName, uint16 InPort)
+: Name(InName)
+, Port(InPort)
+{
+	if (InName.EndsWith(TEXT(".sock")))
+	{
+		checkf(Port == 0, TEXT("Host name appears to be a Unix socket path (contains slashes) and therefore port should be 0"));
+		SocketType = ESocketType::Unix;
+	}
+}
+
+FStringView FZenServiceEndpoint::GetHostName() const
+{
+	return (GetSocketType() == ESocketType::Tcp) ? Name : TEXTVIEW("localhost");
+}
+
+FString FZenServiceEndpoint::GetURL() const
+{
+	TStringBuilder<128> Builder;
+	Builder << TEXTVIEW("http://");
+	Builder << GetHostName();
+	if (Port > 0)
+	{
+		Builder << TEXT(":") << Port;
+	}
+	return FString(Builder);
+}
+
 struct FZenServiceLink
 {
 	FString ServicePath;
@@ -1057,6 +1085,14 @@ FServiceSettings::ReadFromURL(FStringView InstanceURL)
 		InstanceURL.RightChopInline(7);
 	}
 
+	// Is the given URL perhaps a Unix domain socket path?
+	if (InstanceURL.EndsWith(TEXT(".sock")))
+	{
+		ConnectExistingSettings.HostName = InstanceURL;
+		ConnectExistingSettings.Port = 0;
+		return true;
+	}
+
 	int32 PortDelimIndex = INDEX_NONE;
 	InstanceURL.FindLastChar(TEXT(':'), PortDelimIndex);
 	if (PortDelimIndex != INDEX_NONE)
@@ -1868,7 +1904,7 @@ const FString FZenServiceInstance::GetPath() const
 	{
 		return Settings.SettingsVariant.Get<FServiceAutoLaunchSettings>().DataPath;
 	}
-	return GetURL();
+	return GetEndpoint().GetURL();
 }
 
 bool
@@ -1891,14 +1927,12 @@ FZenServiceInstance::PingService(uint32_t AttemptCount)
 	{
 		Attempt++;
 
-		TStringBuilder<128> ZenDomain;
-		ZenDomain << HostName << TEXT(":") << Port;
-		Zen::FZenHttpRequest Request(ZenDomain, false, 500);
+		Zen::FZenHttpRequest Request(Endpoint, false, 500);
 		Zen::FZenHttpRequest::Result Result = Request.PerformBlockingDownload(TEXTVIEW("health/ready"), nullptr, Zen::EContentType::Text);
 		
 		if (Result == Zen::FZenHttpRequest::Result::Success && Zen::IsSuccessCode(Request.GetResponseCode()))
 		{
-			UE_LOGFMT(LogZenServiceInstance, Display, "Unreal Zen Storage Server HTTP service at {Domain} status: {Response}.", ZenDomain, Request.GetResponseAsString());
+			UE_LOGFMT(LogZenServiceInstance, Display, "Unreal Zen Storage Server HTTP service at {Domain} status: {Response}.", Endpoint.GetURL(), Request.GetResponseAsString());
 			return true;
 		}
 
@@ -1906,13 +1940,13 @@ FZenServiceInstance::PingService(uint32_t AttemptCount)
 		{
 			if (Attempt == AttemptCount)
 			{
-				UE_LOGFMT(LogZenServiceInstance, Warning, "Unable to reach Unreal Zen Storage Server HTTP service at {Domain}. Status: {Code}. Response: {Response}", ZenDomain, Request.GetResponseCode(), Request.GetResponseAsString());
+				UE_LOGFMT(LogZenServiceInstance, Warning, "Unable to reach Unreal Zen Storage Server HTTP service at {Domain}. Status: {Code}. Response: {Response}", Endpoint.GetURL(), Request.GetResponseCode(), Request.GetResponseAsString());
 				break;
 			}
 		}
 		else
 		{
-			UE_LOGFMT(LogZenServiceInstance, Display, "Unable to reach Unreal Zen Storage Server HTTP service at {Domain}. Status: {Code}. Response: {Response}", ZenDomain, Request.GetResponseCode(), Request.GetResponseAsString());
+			UE_LOGFMT(LogZenServiceInstance, Display, "Unable to reach Unreal Zen Storage Server HTTP service at {Domain}. Status: {Code}. Response: {Response}", Endpoint.GetURL(), Request.GetResponseCode(), Request.GetResponseAsString());
 			break;
 		}
 	}
@@ -1960,6 +1994,8 @@ FZenServiceInstance::TryRecovery()
 			{
 				UE_LOG(LogZenServiceInstance, Display, TEXT("Local ZenServer recovery being attempted..."));
 
+				uint16 Port = GetEndpoint().GetPort();
+
 				bool bShutdownExistingInstance = true;
 				{
 					const ZenServerState ServerState(/* ReadOnly */true);
@@ -1978,7 +2014,8 @@ FZenServiceInstance::TryRecovery()
 					return false;
 				}
 
-				AutoLaunch(Settings.SettingsVariant.Get<FServiceAutoLaunchSettings>(), ExecutablePath, ExecutionContextFilePath, HostName, Port);
+				AutoLaunch(Settings.SettingsVariant.Get<FServiceAutoLaunchSettings>(), ExecutablePath, ExecutionContextFilePath, Endpoint);
+				Port = GetEndpoint().GetPort();
 
 				FDateTime StartedWaitingForHealth = FDateTime::UtcNow();
 				bLastRecoveryResult = PingService(1);
@@ -1994,7 +2031,7 @@ FZenServiceInstance::TryRecovery()
 					FPlatformProcess::Sleep(0.5f);
 					if (!IsZenProcessUsingEffectivePort(Port))
 					{
-						AutoLaunch(Settings.SettingsVariant.Get<FServiceAutoLaunchSettings>(), ExecutablePath, ExecutionContextFilePath, HostName, Port);
+						AutoLaunch(Settings.SettingsVariant.Get<FServiceAutoLaunchSettings>(), ExecutablePath, ExecutionContextFilePath, Endpoint);
 					}
 					bLastRecoveryResult = PingService(1);
 				}
@@ -2018,6 +2055,7 @@ FZenServiceInstance::TryRecovery()
 bool
 FZenServiceInstance::AddSponsorProcessIDs(TArrayView<uint32> SponsorProcessIDs)
 {
+	uint16 Port = GetEndpoint().GetPort();
 	ZenServerState State(/*bReadOnly*/ false);
 	ZenServerState::ZenServerEntry* Entry = State.LookupByEffectiveListenPort(Port);
 	if (Entry)
@@ -2068,7 +2106,9 @@ FZenServiceInstance::Initialize()
 					}
 					else
 					{
-						bHasLaunchedLocal = AutoLaunch(AutoLaunchSettings, ExecutablePath, ExecutionContextFilePath, HostName, Port);
+						uint16 Port = GetEndpoint().GetPort();
+
+						bHasLaunchedLocal = AutoLaunch(AutoLaunchSettings, ExecutablePath, ExecutionContextFilePath, Endpoint);
 						if (bHasLaunchedLocal)
 						{
 							const ZenServerState State(/*ReadOnly*/true);
@@ -2077,7 +2117,7 @@ FZenServiceInstance::Initialize()
 							{
 								AutoLaunchedPid = RunningEntry->Pid.load(std::memory_order_relaxed);
 							}
-							AutoLaunchedPort = Port;
+							AutoLaunchedPort = GetEndpoint().GetPort();
 							bIsRunningLocally = true;
 							if (PingService(2))
 							{
@@ -2106,11 +2146,9 @@ FZenServiceInstance::Initialize()
 	else
 	{
 		const FServiceConnectSettings& ConnectExistingSettings = Settings.SettingsVariant.Get<FServiceConnectSettings>();
-		HostName = ConnectExistingSettings.HostName;
-		Port = ConnectExistingSettings.Port;
-		bIsRunningLocally = IsLocalHost(HostName);
+		bIsRunningLocally = IsLocalHost(ConnectExistingSettings.HostName);
+		Endpoint = FZenServiceEndpoint(ConnectExistingSettings.HostName, ConnectExistingSettings.Port);
 	}
-	URL = WriteToString<64>(TEXT("http://"), HostName, TEXT(":"), Port, TEXT("/"));
 }
 
 static void
@@ -2426,7 +2464,7 @@ FZenServiceInstance::WriteLocalPluginsConfig()
 }
 
 bool
-FZenServiceInstance::AutoLaunch(const FServiceAutoLaunchSettings& InSettings, const FString& ExecutablePath, const FString& ExecutionContextFilePath, FString& OutHostName, uint16& OutPort)
+FZenServiceInstance::AutoLaunch(const FServiceAutoLaunchSettings& InSettings, const FString& ExecutablePath, const FString& ExecutionContextFilePath, FZenServiceEndpoint& OutEndpoint)
 {
 	IFileManager& FileManager = IFileManager::Get();
 	const FString LockFilePath = FPaths::Combine(InSettings.DataPath, TEXT(".lock"));
@@ -2616,9 +2654,8 @@ FZenServiceInstance::AutoLaunch(const FServiceAutoLaunchSettings& InSettings, co
 		GConfig->SetString(TEXT("/Script/UnrealEd.ZenServerSettings"), TEXT("InUseDefaultDataPath"), *InSettings.DataPath, GEditorSettingsIni);
 	}
 
-	OutHostName = TEXT("[::1]");
-	// Default to assuming that we get to run on the port we want
-	OutPort = InSettings.DesiredPort;
+	FStringView HostName = TEXT("[::1]");
+	OutEndpoint = FZenServiceEndpoint(HostName, InSettings.DesiredPort);
 
 	ZenLockFileData RunningLockFileState = ZenLockFileData::ReadCbLockFile(*LockFilePath);
 	if (!RunningLockFileState.IsValid)
@@ -2629,7 +2666,8 @@ FZenServiceInstance::AutoLaunch(const FServiceAutoLaunchSettings& InSettings, co
 	{
 		return false;
 	}
-	OutPort = RunningLockFileState.EffectivePort;
+
+	OutEndpoint = FZenServiceEndpoint(HostName, RunningLockFileState.EffectivePort);
 	return true;
 }
 
@@ -2660,18 +2698,15 @@ FZenServiceInstance::GetCacheStats(FZenCacheStats& Stats)
 		}
 	}
 
-	TStringBuilder<128> ZenDomain;
-	ZenDomain << HostName << TEXT(":") << Port;
-
 #if WITH_EDITOR
 	EAsyncExecution ThreadPool = EAsyncExecution::LargeThreadPool;
 #else
 	EAsyncExecution ThreadPool = EAsyncExecution::ThreadPool;
 #endif
 		// We've not got any requests in flight and we've met a given time requirement for requests
-	CacheStatsRequest = Async(ThreadPool, [ZenDomain = FString(ZenDomain)]
+	CacheStatsRequest = Async(ThreadPool, [&Endpoint = Endpoint]
 		{
-			UE::Zen::FZenHttpRequest Request(ZenDomain, false);
+			UE::Zen::FZenHttpRequest Request(Endpoint, false);
 
 			TArray64<uint8> GetBuffer;
 			FZenHttpRequest::Result Result = Request.PerformBlockingDownload(TEXTVIEW("/stats/z$"), &GetBuffer, Zen::EContentType::CbObject);
@@ -2717,18 +2752,15 @@ FZenServiceInstance::GetProjectStats(FZenProjectStats& Stats)
 		}
 	}
 
-	TStringBuilder<128> ZenDomain;
-	ZenDomain << HostName << TEXT(":") << Port;
-
 #if WITH_EDITOR
 	EAsyncExecution ThreadPool = EAsyncExecution::LargeThreadPool;
 #else
 	EAsyncExecution ThreadPool = EAsyncExecution::ThreadPool;
 #endif
 			// We've not got any requests in flight and we've met a given time requirement for requests
-	ProjectStatsRequest = Async(ThreadPool, [ZenDomain = FString(ZenDomain)]
+	ProjectStatsRequest = Async(ThreadPool, [&Endpoint = Endpoint]
 		{
-			UE::Zen::FZenHttpRequest Request(ZenDomain, false);
+			UE::Zen::FZenHttpRequest Request(Endpoint, false);
 
 			TArray64<uint8> GetBuffer;
 			FZenHttpRequest::Result Result = Request.PerformBlockingDownload(TEXTVIEW("/stats/prj"), &GetBuffer, Zen::EContentType::CbObject);
@@ -2778,13 +2810,10 @@ FZenServiceInstance::GetGCStatus(FGCStatus& Status)
 #else
 		EAsyncExecution ThreadPool = EAsyncExecution::ThreadPool;
 #endif
-		TStringBuilder<128> ZenDomain;
-		ZenDomain << HostName << TEXT(":") << Port;
-
 		// We've not got any requests in flight and we've met a given time requirement for requests
-		GCStatusRequest = Async(ThreadPool, [this, ZenDomain = FString(ZenDomain)]
+		GCStatusRequest = Async(ThreadPool, [this]
 			{
-				UE::Zen::FZenHttpRequest Request(ZenDomain, false);
+				UE::Zen::FZenHttpRequest Request(Endpoint, false);
 
 				TArray64<uint8> GetBuffer;
 				FZenHttpRequest::Result Result = Request.PerformBlockingDownload(TEXTVIEW("/admin/gc"), &GetBuffer, Zen::EContentType::CbObject);
@@ -2809,7 +2838,7 @@ FZenServiceInstance::GetGCStatus(FGCStatus& Status)
 bool 
 FZenServiceInstance::RequestGC(const bool* OverrideCollectSmallObjects, const uint32* OverrideMaxCacheDuration)
 {
-	UE::Zen::FZenHttpRequest Request(WriteToString<128>(HostName, ':', Port), /*bLogErrors*/ false);
+	UE::Zen::FZenHttpRequest Request(Endpoint, /*bLogErrors*/ false);
 
 	TCHAR Separators[] = {TEXT('?'), TEXT('&')};
 	int32 SeparatorIndex = 0;
@@ -3096,18 +3125,15 @@ FZenServiceInstance::GatherAnalytics(TArray<FAnalyticsEventAttribute>& Attribute
 bool
 FZenServiceInstance::GetWorkspaces(FZenWorkspaces& Workspaces) const
 {
-	TStringBuilder<128> ZenDomain;
-	ZenDomain << HostName << TEXT(":") << Port;
-
 #if WITH_EDITOR
 	EAsyncExecution ThreadPool = EAsyncExecution::LargeThreadPool;
 #else
 	EAsyncExecution ThreadPool = EAsyncExecution::ThreadPool;
 #endif
 
-	WorkspacesRequest = Async(ThreadPool, [ZenDomain = FString(ZenDomain)]
+	WorkspacesRequest = Async(ThreadPool, [&Endpoint = Endpoint]
 	{
-		UE::Zen::FZenHttpRequest Request(ZenDomain, false);
+		UE::Zen::FZenHttpRequest Request(Endpoint, false);
 
 		FZenHttpRequest::Result Result = Request.PerformBlockingDownload(TEXTVIEW("/ws"), nullptr, Zen::EContentType::CbObject);
 
diff --git a/Engine/Source/Developer/Zen/Public/Experimental/ZenServerInterface.h b/Engine/Source/Developer/Zen/Public/Experimental/ZenServerInterface.h
--- a/Engine/Source/Developer/Zen/Public/Experimental/ZenServerInterface.h
+++ b/Engine/Source/Developer/Zen/Public/Experimental/ZenServerInterface.h
@@ -211,6 +211,29 @@ struct FGCStatus
 	FString Description;
 };
 
+class FZenServiceEndpoint
+{
+public:
+	enum class ESocketType : uint8
+	{
+		Tcp,
+		Unix,
+	};
+
+	FZenServiceEndpoint() = default;
+	FZenServiceEndpoint(FStringView InName, uint16 InPort);
+	ESocketType GetSocketType() const { return SocketType; }
+	FStringView GetName() const { return Name; }
+	UE_API FStringView GetHostName() const;
+	UE_API FString GetURL() const;
+	uint16 GetPort() const { return Port; }
+
+private:
+	FString Name = TEXT("localhost");
+	uint16 Port = 8558;
+	ESocketType SocketType = ESocketType::Tcp;
+};
+
 /**
  * Type used to declare usage of a Zen server instance whether the shared default instance or a unique non-default instance.
  * Used to help manage launch, and optionally in the future, shutdown a shared default instance.  Use the default constructor
@@ -264,13 +287,11 @@ public:
 	UE_API ~FZenServiceInstance();
 
 	UE_API const FString GetPath() const;
-	inline const TCHAR* GetURL() const { return *URL; }
-	inline const TCHAR* GetHostName() const { return *HostName; }
-	inline uint16 GetPort() const { return Port; }
 	inline const FServiceSettings& GetServiceSettings() const { return Settings; }
 	UE_API bool IsServiceRunning();
 	UE_API bool IsServiceReady();
 	bool IsServiceRunningLocally() const { return bIsRunningLocally; }
+	const FZenServiceEndpoint& GetEndpoint() const { return Endpoint; }
 
 	UE_API bool TryRecovery();
 
@@ -293,7 +314,7 @@ private:
 	void Initialize();
 	bool ConditionalUpdateLocalInstall(FServiceAutoLaunchSettings::EInstallMode InstallMode);
 	bool WriteLocalPluginsConfig();
-	static bool AutoLaunch(const FServiceAutoLaunchSettings& InSettings, const FString& ExecutablePath, const FString& ExecutionContextFilePath, FString& OutHostName, uint16& OutPort);
+	static bool AutoLaunch(const FServiceAutoLaunchSettings& InSettings, const FString& ExecutablePath, const FString& ExecutionContextFilePath, FZenServiceEndpoint& OutEndpoint);
 	bool PingService(uint32_t AttemptCount);
 
 	mutable TFuture<TOptional<FGCStatus>> GCStatusRequest;
@@ -313,9 +334,8 @@ private:
 	mutable TFuture<FZenWorkspaces> WorkspacesRequest;
 
 	FServiceSettings Settings;
-	FString URL;
-	FString HostName;
-	uint16 Port;
+
+	FZenServiceEndpoint Endpoint;
 	static uint16 AutoLaunchedPort;
 	static uint32 AutoLaunchedPid;
 	bool bHasLaunchedLocal = false;
diff --git a/Engine/Source/Editor/UnrealEd/Private/DerivedDataBuildRemoteExecutor.cpp b/Engine/Source/Editor/UnrealEd/Private/DerivedDataBuildRemoteExecutor.cpp
--- a/Engine/Source/Editor/UnrealEd/Private/DerivedDataBuildRemoteExecutor.cpp
+++ b/Engine/Source/Editor/UnrealEd/Private/DerivedDataBuildRemoteExecutor.cpp
@@ -165,7 +165,7 @@ public:
 			
 			if (bEnabled)
 			{
-				RequestPool = MakeUnique<UE::Zen::FZenHttpRequestPool>(ZenServiceInstance.GetURL());
+				RequestPool = MakeUnique<UE::Zen::FZenHttpRequestPool>(ZenServiceInstance.GetEndpoint());
 				IModularFeatures::Get().RegisterModularFeature(IBuildWorkerExecutor::FeatureName, this);
 			}
 			else
diff --git a/Engine/Source/Editor/ZenEditor/Private/SZenStatusBar.cpp b/Engine/Source/Editor/ZenEditor/Private/SZenStatusBar.cpp
--- a/Engine/Source/Editor/ZenEditor/Private/SZenStatusBar.cpp
+++ b/Engine/Source/Editor/ZenEditor/Private/SZenStatusBar.cpp
@@ -191,7 +191,8 @@ void FZenStausBarCommands::ViewZenStore_Clicked()
 
 	if (ServerInstance.IsValid())
 	{
-		FPlatformProcess::LaunchURL(*FString::Printf(TEXT("%s/dashboard/?"), ServerInstance->GetURL()), nullptr, nullptr);
+		FString URL = ServerInstance->GetEndpoint().GetURL();
+		FPlatformProcess::LaunchURL(*FString::Printf(TEXT("%s/dashboard/?"), *URL), nullptr, nullptr);
 	}
 }
 
