Fix bug where zenserver will not start if UE install path contains special characters.

diff --git a/Engine/Source/Developer/Zen/Private/ZenServerInterface.cpp b/Engine/Source/Developer/Zen/Private/ZenServerInterface.cpp
--- a/Engine/Source/Developer/Zen/Private/ZenServerInterface.cpp
+++ b/Engine/Source/Developer/Zen/Private/ZenServerInterface.cpp
@@ -196,7 +196,28 @@ GetServiceLinkPath()
 static FString
 GetServicePluginsConfigPath()
 {
-	return FPaths::Combine(GetLocalInstallPath(), TEXT("zen_plugins.json"));
+	const int PluginsConfigVersion = 1;
+	return FPaths::Combine(GetLocalInstallPath(), FString::Printf(TEXT("zen_plugins_v%d.json"), PluginsConfigVersion));
+}
+
+static void
+CleanOutOfDateServicePluginConfigs()
+{
+	IFileManager::Get().IterateDirectory(*GetLocalInstallPath(), [](const TCHAR* Pathname, bool bIsDirectory)
+	{
+		if (!bIsDirectory)
+		{
+			FString Path = Pathname;
+			FString FileName = FPaths::GetCleanFilename(Path);
+
+			if (FileName.EndsWith(TEXT(".json")) && FileName.StartsWith("zen_plugins") &&
+				FileName != FPaths::GetCleanFilename(GetServicePluginsConfigPath()))
+			{
+				IFileManager::Get().Delete(Pathname, false /*RequireExists*/, false /*EvenReadOnly*/, true /*Quiet*/);
+			}
+		}
+		return true;
+	});
 }
 
 static FString
@@ -2217,6 +2238,7 @@ FZenServiceInstance::ConditionalUpdateLocalInstall(FServiceAutoLaunchSettings::E
 	if (bForceZenPluginsInstall || bMissingZenPluginsConfig)
 	{
 		WriteLocalPluginsConfig();
+		CleanOutOfDateServicePluginConfigs();
 	}
 
 	if (InstallMode == FServiceAutoLaunchSettings::EInstallMode::Link)
@@ -2460,7 +2482,7 @@ FZenServiceInstance::WriteLocalPluginsConfig()
 
 	const FString OutputPath = GetServicePluginsConfigPath();
 	UE_LOG(LogZenServiceInstance, Display, TEXT("Writing plugin configuration to '%s'"), *OutputPath);
-	return FFileHelper::SaveStringToFile(OutputString, *OutputPath);
+	return FFileHelper::SaveStringToFile(OutputString, *OutputPath, FFileHelper::EEncodingOptions::ForceUTF8WithoutBOM);
 }
 
 bool
