Mac compile fixes for last submit.

diff --git a/Engine/Source/Developer/IoStoreUtilities/Private/ZenStoreWriter.cpp b/Engine/Source/Developer/IoStoreUtilities/Private/ZenStoreWriter.cpp
--- a/Engine/Source/Developer/IoStoreUtilities/Private/ZenStoreWriter.cpp
+++ b/Engine/Source/Developer/IoStoreUtilities/Private/ZenStoreWriter.cpp
@@ -496,7 +496,7 @@ void FZenStoreWriter::Initialize(const FCookInfo& Info)
 		ServicePort = ZenServiceInstance.GetEndpoint().GetPort();
 #endif
 
-		ProjectStoreWriter.Write(ServiceHostName, ServicePort, IsRunningLocally, *ProjectId, *OplogId);
+		ProjectStoreWriter.Write(ServiceHostName, ServicePort, IsRunningLocally, *ProjectId, *OplogId, TargetPlatformFName);
 
 		if (!bOplogEstablished && CleanBuild)
 		{
diff --git a/Engine/Source/Developer/Zen/Private/BuildServerInterface.cpp b/Engine/Source/Developer/Zen/Private/BuildServerInterface.cpp
--- a/Engine/Source/Developer/Zen/Private/BuildServerInterface.cpp
+++ b/Engine/Source/Developer/Zen/Private/BuildServerInterface.cpp
@@ -714,7 +714,7 @@ FBuildServiceInstance::BuildTransferThreadLoop()
 														TEXT("ue.projectstore"));
 
 					Zen::FZenProjectStoreWriter ProjectStoreWriter(*OplogLifetimeMarkerPath);
-					ProjectStoreWriter.Write(TEXT("[::1]"), LocalPort, true, DestinationProjectId, DestinationOplogId);
+					ProjectStoreWriter.Write(TEXT("[::1]"), LocalPort, true, DestinationProjectId, DestinationOplogId, NAME_None);
 
 					FString RootDir = FPaths::RootDir();
 					FString EngineDir = FPaths::EngineDir();
diff --git a/Engine/Source/Developer/Zen/Private/ZenProjectStoreWriter.cpp b/Engine/Source/Developer/Zen/Private/ZenProjectStoreWriter.cpp
--- a/Engine/Source/Developer/Zen/Private/ZenProjectStoreWriter.cpp
+++ b/Engine/Source/Developer/Zen/Private/ZenProjectStoreWriter.cpp
@@ -7,8 +7,8 @@
 #include "Containers/StringView.h"
 #include "Experimental/ZenServerInterface.h"
 #include "HAL/FileManager.h"
+#include "HAL/PlatformMisc.h"
 #include "Policies/PrettyJsonPrintPolicy.h"
-#include "Serialization/Archive.h"
 #include "Serialization/JsonSerializer.h"
 #include "Serialization/JsonWriter.h"
 #include "SocketSubsystem.h"
@@ -21,7 +21,7 @@ FZenProjectStoreWriter::FZenProjectStoreWriter(const TCHAR* TargetFileName)
 {
 }
 
-bool FZenProjectStoreWriter::Write(FStringView ServiceHostName, uint16 ServicePort, bool bIsRunningLocally, FStringView ProjectId, FStringView OplogId)
+bool FZenProjectStoreWriter::Write(FStringView ServiceHostName, uint16 ServicePort, bool bIsRunningLocally, FStringView ProjectId, FStringView OplogId, FName TargetPlatformName)
 {
 	if (!OplogMarker)
 	{
@@ -62,7 +62,7 @@ bool FZenProjectStoreWriter::Write(FStringView ServiceHostName, uint16 ServicePo
 				// Store the Bonjour hostname when on a mac so that the client can attempt to connect via USB.
 				// This is done since the "Link Local" IP of the USB connection changes everytime
 				// the cable is plugged in, so there is no stable IP to save.
-				bool bAppleTarget = TargetPlatformFName == TEXT("IOS") || TargetPlatformFName == TEXT("TVOS");
+				bool bAppleTarget = TargetPlatformName == TEXT("IOS") || TargetPlatformName == TEXT("TVOS");
 				if (bAppleTarget)
 				{
 					FString MacBonjourName = FPlatformMisc::GetBonjourName();
diff --git a/Engine/Source/Developer/Zen/Public/Experimental/ZenProjectStoreWriter.h b/Engine/Source/Developer/Zen/Public/Experimental/ZenProjectStoreWriter.h
--- a/Engine/Source/Developer/Zen/Public/Experimental/ZenProjectStoreWriter.h
+++ b/Engine/Source/Developer/Zen/Public/Experimental/ZenProjectStoreWriter.h
@@ -20,7 +20,7 @@ public:
 	UE_API FZenProjectStoreWriter(const TCHAR* TargetFileName);
 	UE_API ~FZenProjectStoreWriter() = default;
 
-	UE_API bool Write(FStringView ServiceHostName, uint16 ServicePort, bool bIsRunningLocally, FStringView ProjectId, FStringView OplogId);
+	UE_API bool Write(FStringView ServiceHostName, uint16 ServicePort, bool bIsRunningLocally, FStringView ProjectId, FStringView OplogId, FName TargetPlatformName);
 private:
 	TUniquePtr<FArchive> OplogMarker;
 };
