Ensure that BuildStorageTool is supplying an appropiate TargetPlatformName that is distinct from the cook platform.

diff --git a/Engine/Source/Developer/Zen/Private/BuildServerInterface.cpp b/Engine/Source/Developer/Zen/Private/BuildServerInterface.cpp
--- a/Engine/Source/Developer/Zen/Private/BuildServerInterface.cpp
+++ b/Engine/Source/Developer/Zen/Private/BuildServerInterface.cpp
@@ -349,7 +349,7 @@ FBuildServiceInstance::StartBuildTransfer(const FCbObjectId& BuildId, FStringVie
 }
 
 FBuildServiceInstance::FBuildTransfer
-FBuildServiceInstance::StartOplogBuildTransfer(const FCbObjectId& BuildId, FStringView DestinationProjectId, FStringView DestinationOplogId, FStringView ProjectFilePath, FStringView Namespace, FStringView Bucket)
+FBuildServiceInstance::StartOplogBuildTransfer(const FCbObjectId& BuildId, FStringView DestinationProjectId, FStringView DestinationOplogId, FStringView ProjectFilePath, FStringView Namespace, FStringView Bucket, FName TargetPlatformName)
 {
 	const FString ZenUtility = UE::Zen::GetLocalInstallUtilityPath();
 	TSharedPtr<FBuildTransferCommand> NewTransferCommand = MakeShared<FBuildTransferCommand>();
@@ -359,6 +359,7 @@ FBuildServiceInstance::StartOplogBuildTransfer(const FCbObjectId& BuildId, FStri
 	NewTransferCommand->Bucket = Bucket;
 	NewTransferCommand->Destination = *WriteToString<64>(DestinationProjectId, TEXT("/"), DestinationOplogId);
 	NewTransferCommand->ProjectFilePath = ProjectFilePath;
+	NewTransferCommand->TargetPlatformName = TargetPlatformName;
 	BuildTransferThreadCommands.Enqueue(NewTransferCommand.ToSharedRef());
 	KickBuildTransferThread();
 
@@ -714,7 +715,7 @@ FBuildServiceInstance::BuildTransferThreadLoop()
 														TEXT("ue.projectstore"));
 
 					Zen::FZenProjectStoreWriter ProjectStoreWriter(*OplogLifetimeMarkerPath);
-					ProjectStoreWriter.Write(TEXT("[::1]"), LocalPort, true, DestinationProjectId, DestinationOplogId, NAME_None);
+					ProjectStoreWriter.Write(TEXT("[::1]"), LocalPort, true, DestinationProjectId, DestinationOplogId, TransferCommand->TargetPlatformName);
 
 					FString RootDir = FPaths::RootDir();
 					FString EngineDir = FPaths::EngineDir();
diff --git a/Engine/Source/Developer/Zen/Public/Experimental/BuildServerInterface.h b/Engine/Source/Developer/Zen/Public/Experimental/BuildServerInterface.h
--- a/Engine/Source/Developer/Zen/Public/Experimental/BuildServerInterface.h
+++ b/Engine/Source/Developer/Zen/Public/Experimental/BuildServerInterface.h
@@ -67,6 +67,7 @@ private:
 		FString Bucket;
 		FString Destination;
 		FString ProjectFilePath;
+		FName TargetPlatformName = NAME_None;
 
 		std::atomic<bool> bCancelRequested = false;
 
@@ -125,7 +126,7 @@ public:
 	UE_API void ListBuilds(FStringView Namespace, FStringView Bucket, FOnListBuildsComplete&& InOnListBuildsComplete);
 
 	UE_API FBuildTransfer StartBuildTransfer(const FCbObjectId& BuildId, FStringView DestinationFolder, FStringView Namespace, FStringView Bucket);
-	UE_API FBuildTransfer StartOplogBuildTransfer(const FCbObjectId& BuildId, FStringView DestinationProjectId, FStringView DestinationOplogId, FStringView ProjectFilePath, FStringView Namespace, FStringView Bucket);
+	UE_API FBuildTransfer StartOplogBuildTransfer(const FCbObjectId& BuildId, FStringView DestinationProjectId, FStringView DestinationOplogId, FStringView ProjectFilePath, FStringView Namespace, FStringView Bucket, FName TargetPlatformName);
 private:
 	typedef TUniqueFunction<void(FString& CommandLineArgs)> FPreZenUtilityInvocation;
 	typedef TUniqueFunction<void(bool)> FOnZenUtilityInvocationComplete;
diff --git a/Engine/Source/Editor/StorageServerWidgets/Private/SBuildSelection.cpp b/Engine/Source/Editor/StorageServerWidgets/Private/SBuildSelection.cpp
--- a/Engine/Source/Editor/StorageServerWidgets/Private/SBuildSelection.cpp
+++ b/Engine/Source/Editor/StorageServerWidgets/Private/SBuildSelection.cpp
@@ -1550,6 +1550,12 @@ TSharedRef<SWidget> SBuildSelection::GetGridPanel()
 									DestinationPlatformName = *WriteToString<64>(CookPlatformField.AsString());
 								}
 
+								FName TargetPlatformName = NAME_None;
+								if (FCbFieldView PlatformField = BuildRecord->Metadata["platform"]; PlatformField.HasValue() && !PlatformField.HasError())
+								{
+									TargetPlatformName = FName(PlatformField.AsString());
+								}
+
 								EBuildType BuildType = GetSelectedBuildType();
 								int BuildTypeIndex = (int)BuildType;
 								const bool bAppendBuildNameToDestination = bAppendBuildNameToDestinations[BuildTypeIndex];
@@ -1564,7 +1570,7 @@ TSharedRef<SWidget> SBuildSelection::GetGridPanel()
 										DestinationOplogId = FString::Printf(TEXT("%s.%s"), *DestinationOplogId, *SanitizeForZenId(SelectedItems[0]->DisplayName));
 									}
 									FBuildServiceInstance::FBuildTransfer BuildTransfer =
-										ServiceInstance->StartOplogBuildTransfer(BuildRecord->BuildId, DestinationProjectId, DestinationOplogId, ProjectFilePath, SelectedItems[0]->Namespace, Bucket);
+										ServiceInstance->StartOplogBuildTransfer(BuildRecord->BuildId, DestinationProjectId, DestinationOplogId, ProjectFilePath, SelectedItems[0]->Namespace, Bucket, TargetPlatformName);
 									OnBuildTransferStarted.ExecuteIfBound(BuildTransfer, SelectedItems[0]->DisplayName, PlatformForSelectedGroup);
 								}
 								else
