Prepping for BuildStorageTool logging changes when launching subprocesses.

diff --git a/Engine/Source/Developer/Zen/Private/BuildServerInterface.cpp b/Engine/Source/Developer/Zen/Private/BuildServerInterface.cpp
--- a/Engine/Source/Developer/Zen/Private/BuildServerInterface.cpp
+++ b/Engine/Source/Developer/Zen/Private/BuildServerInterface.cpp
@@ -377,12 +377,11 @@ FBuildServiceInstance::FBuildTransfer::GetDescription() const
 	}
 
 	FReadScopeLock _(BuildTransferCommand->Lock);
-	int32 Index = 0;
-	if (BuildTransferCommand->Output.FindLastChar(TCHAR('\n'), Index))
+	if (BuildTransferCommand->NextOutputLineIndex == 0)
 	{
-		return BuildTransferCommand->Output.RightChop(Index+1);
+		return FString();
 	}
-	return BuildTransferCommand->Output;
+	return BuildTransferCommand->RecentOutputLines[BuildTransferCommand->RecentOutputLines.GetPreviousIndex(BuildTransferCommand->NextOutputLineIndex)];
 }
 
 FBuildServiceInstance::EBuildTransferStatus
@@ -760,8 +759,7 @@ FBuildServiceInstance::BuildTransferThreadLoop()
 		MonitoredUtilityProcess.OnOutput().BindSPLambda(TransferCommand, [TransferCommand](FString Output)
 		{
 			FWriteScopeLock _(TransferCommand->Lock);
-			TransferCommand->Output += TEXT("\n");
-			TransferCommand->Output += Output;
+			TransferCommand->RecentOutputLines[TransferCommand->NextOutputLineIndex++] = MoveTemp(Output);
 		});
 
 		{
diff --git a/Engine/Source/Developer/Zen/Public/Experimental/BuildServerInterface.h b/Engine/Source/Developer/Zen/Public/Experimental/BuildServerInterface.h
--- a/Engine/Source/Developer/Zen/Public/Experimental/BuildServerInterface.h
+++ b/Engine/Source/Developer/Zen/Public/Experimental/BuildServerInterface.h
@@ -2,6 +2,7 @@
 
 #pragma once
 
+#include "Containers/CircularBuffer.h"
 #include "Containers/Map.h"
 #include "Containers/MpscQueue.h"
 #include "Containers/StringFwd.h"
@@ -72,6 +73,9 @@ private:
 		std::atomic<bool> bCancelRequested = false;
 
 		FRWLock Lock;
+		const uint32 RecentOutputLinesCapacity = 16;
+		TCircularBuffer<FString> RecentOutputLines = TCircularBuffer<FString>(RecentOutputLinesCapacity);
+		uint32 NextOutputLineIndex = 0;
 		FString Output;
 		EBuildTransferStatus Status = EBuildTransferStatus::Queued;
 		int ReturnCode;
