BuildStorageTool usability improvements

diff --git a/Engine/Source/Developer/Zen/Private/BuildServerInterface.cpp b/Engine/Source/Developer/Zen/Private/BuildServerInterface.cpp
--- a/Engine/Source/Developer/Zen/Private/BuildServerInterface.cpp
+++ b/Engine/Source/Developer/Zen/Private/BuildServerInterface.cpp
@@ -10,12 +10,15 @@
 #include "HAL/CriticalSection.h"
 #include "HAL/FileManager.h"
 #include "HAL/PlatformFileManager.h"
+#include "HAL/PlatformOutputDevices.h"
 #include "Http/HttpClient.h"
 #include "Http/HttpHostBuilder.h"
 #include "Misc/App.h"
 #include "Misc/ConfigCacheIni.h"
 #include "Misc/MonitoredProcess.h"
+#include "Misc/OutputDeviceFile.h"
 #include "Misc/Paths.h"
+#include "Misc/PathViews.h"
 #include "Misc/ScopeExit.h"
 #include "Misc/ScopeLock.h"
 #include "Misc/ScopeRWLock.h"
@@ -42,6 +45,46 @@ DEFINE_LOG_CATEGORY_STATIC(LogBuildServiceInstance, Log, All);
 
 namespace Private
 {
+
+enum class EBuildTransferType
+{
+	Files,
+	Oplog
+};
+
+struct FProgressState
+{
+	FString Label;
+	FString Detail;
+	float Percent = 0.0f;
+};
+
+struct FBuildTransferState : public TSharedFromThis<FBuildTransferState, ESPMode::ThreadSafe>
+{
+	EBuildTransferType Type;
+	FCbObjectId BuildId;
+	FString Name;
+	FString Namespace;
+	FString Bucket;
+	FString Destination;
+	FString ProjectFilePath;
+	FName TargetPlatformName = NAME_None;
+
+	std::atomic<bool> bCancelRequested = false;
+
+	mutable FRWLock Lock;
+	TArray<FProgressState> ProgressState;
+	const uint32 RecentOutputLinesCapacity = 16;
+	TCircularBuffer<FString> RecentOutputLines = TCircularBuffer<FString>(RecentOutputLinesCapacity);
+	uint32 NextOutputLineIndex = 0;
+	TUniquePtr<FArchive> OutputLogFile;
+	UE::Zen::Build::FBuildServiceInstance::EBuildTransferStatus Status = UE::Zen::Build::FBuildServiceInstance::EBuildTransferStatus::Queued;
+	int ReturnCode;
+
+	FString GetRecentOutput() const;
+	FString GetLogFilename() const;
+};
+
 class FAccessToken
 {
 public:
@@ -61,6 +104,42 @@ FAnsiStringBuilderBase& operator<<(FAnsiStringBuilderBase& Builder, const Privat
 	return Builder.Append(Token.Header);
 }
 
+FString
+FBuildTransferState::GetRecentOutput() const
+{
+	TStringBuilder<256> StringBuilder;
+	if (NextOutputLineIndex == 0)
+	{
+		return FString();
+	}
+
+	uint32 StartIndex = 0;
+	if (NextOutputLineIndex > RecentOutputLinesCapacity)
+	{
+		StartIndex = NextOutputLineIndex - RecentOutputLinesCapacity + 1;
+	}
+
+	for (uint32 Index = StartIndex; Index < NextOutputLineIndex; ++Index)
+	{
+		StringBuilder.Append(RecentOutputLines[Index]);
+		if (Index < (NextOutputLineIndex - 1))
+		{
+			StringBuilder.Append(LINE_TERMINATOR);
+		}
+	}
+
+	return StringBuilder.ToString();
+}
+
+FString
+FBuildTransferState::GetLogFilename() const
+{
+	const FString LogDirectory = IFileManager::Get().ConvertToAbsolutePathForExternalAppForWrite(
+		*FPaths::GetPath(FPlatformOutputDevices::GetAbsoluteLogFilename())
+	);
+	return FPaths::Combine(LogDirectory, TEXT("Transfers"), Name) + TEXT(".log");
+}
+
 void FAccessToken::SetToken(const FStringView Scheme, const FStringView Token)
 {
 	FWriteScopeLock WriteLock(Lock);
@@ -331,81 +410,165 @@ FBuildServiceInstance::ListBuilds(FStringView Namespace, FStringView Bucket, FOn
 }
 
 FBuildServiceInstance::FBuildTransfer
-FBuildServiceInstance::StartBuildTransfer(const FCbObjectId& BuildId, FStringView DestinationFolder, FStringView Namespace, FStringView Bucket)
+FBuildServiceInstance::StartBuildTransfer(const FCbObjectId& BuildId, FStringView Name, FStringView DestinationFolder, FStringView Namespace, FStringView Bucket)
 {
 	const FString ZenUtility = UE::Zen::GetLocalInstallUtilityPath();
-	TSharedPtr<FBuildTransferCommand> NewTransferCommand = MakeShared<FBuildTransferCommand>();
-	NewTransferCommand->Type = EBuildTransferType::Files;
-	NewTransferCommand->BuildId = BuildId;
-	NewTransferCommand->Namespace = Namespace;
-	NewTransferCommand->Bucket = Bucket;
-	NewTransferCommand->Destination = DestinationFolder;
-	BuildTransferThreadCommands.Enqueue(NewTransferCommand.ToSharedRef());
+	TSharedPtr<Private::FBuildTransferState> NewState = MakeShared<Private::FBuildTransferState>();
+	NewState->Type = Private::EBuildTransferType::Files;
+	NewState->BuildId = BuildId;
+	NewState->Name = Name;
+	NewState->Namespace = Namespace;
+	NewState->Bucket = Bucket;
+	NewState->Destination = DestinationFolder;
+	BuildTransferThreadStates.Enqueue(NewState.ToSharedRef());
 	KickBuildTransferThread();
 
 	FBuildTransfer NewTransfer;
-	NewTransfer.BuildTransferCommand = NewTransferCommand;
+	NewTransfer.State = NewState;
 	return NewTransfer;
 }
 
 FBuildServiceInstance::FBuildTransfer
-FBuildServiceInstance::StartOplogBuildTransfer(const FCbObjectId& BuildId, FStringView DestinationProjectId, FStringView DestinationOplogId, FStringView ProjectFilePath, FStringView Namespace, FStringView Bucket, FName TargetPlatformName)
+FBuildServiceInstance::StartOplogBuildTransfer(const FCbObjectId& BuildId, FStringView Name, FStringView DestinationProjectId, FStringView DestinationOplogId, FStringView ProjectFilePath, FStringView Namespace, FStringView Bucket, FName TargetPlatformName)
 {
 	const FString ZenUtility = UE::Zen::GetLocalInstallUtilityPath();
-	TSharedPtr<FBuildTransferCommand> NewTransferCommand = MakeShared<FBuildTransferCommand>();
-	NewTransferCommand->Type = EBuildTransferType::Oplog;
-	NewTransferCommand->BuildId = BuildId;
-	NewTransferCommand->Namespace = Namespace;
-	NewTransferCommand->Bucket = Bucket;
-	NewTransferCommand->Destination = *WriteToString<64>(DestinationProjectId, TEXT("/"), DestinationOplogId);
-	NewTransferCommand->ProjectFilePath = ProjectFilePath;
-	NewTransferCommand->TargetPlatformName = TargetPlatformName;
-	BuildTransferThreadCommands.Enqueue(NewTransferCommand.ToSharedRef());
+	TSharedPtr<Private::FBuildTransferState> NewState = MakeShared<Private::FBuildTransferState>();
+	NewState->Type = Private::EBuildTransferType::Oplog;
+	NewState->BuildId = BuildId;
+	NewState->Name = Name;
+	NewState->Namespace = Namespace;
+	NewState->Bucket = Bucket;
+	NewState->Destination = *WriteToString<64>(DestinationProjectId, TEXT("/"), DestinationOplogId);
+	NewState->ProjectFilePath = ProjectFilePath;
+	NewState->TargetPlatformName = TargetPlatformName;
+	BuildTransferThreadStates.Enqueue(NewState.ToSharedRef());
 	KickBuildTransferThread();
 
 	FBuildTransfer NewTransfer;
-	NewTransfer.BuildTransferCommand = NewTransferCommand;
+	NewTransfer.State = NewState;
 	return NewTransfer;
 }
 
 FString
 FBuildServiceInstance::FBuildTransfer::GetDescription() const
 {
-	if (!BuildTransferCommand)
+	if (!State)
 	{
 		return TEXT("null");
 	}
 
-	FReadScopeLock _(BuildTransferCommand->Lock);
-	if (BuildTransferCommand->NextOutputLineIndex == 0)
+	FReadScopeLock _(State->Lock);
+	if (State->NextOutputLineIndex == 0)
 	{
 		return FString();
 	}
-	return BuildTransferCommand->RecentOutputLines[BuildTransferCommand->RecentOutputLines.GetPreviousIndex(BuildTransferCommand->NextOutputLineIndex)];
+	return State->RecentOutputLines[State->RecentOutputLines.GetPreviousIndex(State->NextOutputLineIndex)];
+}
+
+FString
+FBuildServiceInstance::FBuildTransfer::GetRecentOutput() const
+{
+	if (!State)
+	{
+		return TEXT("null");
+	}
+
+	FReadScopeLock _(State->Lock);
+	return State->GetRecentOutput();
 }
 
 FBuildServiceInstance::EBuildTransferStatus
 FBuildServiceInstance::FBuildTransfer::GetStatus() const
 {
-	if (!BuildTransferCommand)
+	if (!State)
 	{
 		return EBuildTransferStatus::Invalid;
 	}
 
-	FReadScopeLock _(BuildTransferCommand->Lock);
-	return BuildTransferCommand->Status;
+	FReadScopeLock _(State->Lock);
+	return State->Status;
+}
+
+FString
+FBuildServiceInstance::FBuildTransfer::GetLogFilename() const
+{
+	if (!State)
+	{
+		return FString();
+	}
+
+	FReadScopeLock _(State->Lock);
+	return State->GetLogFilename();
+}
+
+bool
+FBuildServiceInstance::FBuildTransfer::GetCurrentProgress(FString& OutLabel, FString& OutDetail, float& OutPercent) const
+{
+	if (!State)
+	{
+		return false;
+	}
+
+	FReadScopeLock _(State->Lock);
+	if (State->ProgressState.IsEmpty())
+	{
+		if (State->NextOutputLineIndex == 0)
+		{
+			OutLabel.Empty();
+		}
+		else
+		{
+			OutLabel = State->RecentOutputLines[State->RecentOutputLines.GetPreviousIndex(State->NextOutputLineIndex)];
+		}
+		OutDetail = State->GetRecentOutput();
+		OutPercent = 0.f;
+		return true;
+	}
+	OutLabel = State->ProgressState.Top().Label;
+	OutDetail = State->ProgressState.Top().Detail;
+	OutPercent = State->ProgressState.Top().Percent;
+	return true;
+}
+
+bool
+FBuildServiceInstance::FBuildTransfer::GetOverallProgress(FString& OutLabel, FString& OutDetail, float& OutPercent) const
+{
+	if (!State)
+	{
+		return false;
+	}
+
+	FReadScopeLock _(State->Lock);
+	if (State->ProgressState.IsEmpty())
+	{
+		if (State->NextOutputLineIndex == 0)
+		{
+			OutLabel.Empty();
+		}
+		else
+		{
+			OutLabel = State->RecentOutputLines[State->RecentOutputLines.GetPreviousIndex(State->NextOutputLineIndex)];
+		}
+		OutDetail = State->GetRecentOutput();
+		OutPercent = 0.f;
+		return true;
+	}
+	OutLabel = State->ProgressState[0].Label;
+	OutDetail = State->ProgressState[0].Detail;
+	OutPercent = State->ProgressState[0].Percent;
+	return true;
 }
 
 void
 FBuildServiceInstance::FBuildTransfer::RequestCancel()
 {
-	if (!BuildTransferCommand)
+	if (!State)
 	{
 		return;
 	}
-	BuildTransferCommand->bCancelRequested.store(true);
-	FWriteScopeLock _(BuildTransferCommand->Lock);
-	BuildTransferCommand->Status = EBuildTransferStatus::Canceled;
+	State->bCancelRequested.store(true);
+	FWriteScopeLock _(State->Lock);
+	State->Status = EBuildTransferStatus::Canceled;
 }
 
 FBuildServiceInstance::EConnectionState
@@ -660,39 +823,41 @@ FBuildServiceInstance::KickBuildTransferThread()
 void
 FBuildServiceInstance::BuildTransferThreadLoop()
 {
-	while (!BuildTransferThreadCommands.IsEmpty() || !bBuildTransferThreadStopping.load(std::memory_order_relaxed))
+	while (!BuildTransferThreadStates.IsEmpty() || !bBuildTransferThreadStopping.load(std::memory_order_relaxed))
 	{
-		TOptional<TSharedRef<FBuildTransferCommand>> OptionalCommand = BuildTransferThreadCommands.Dequeue();
-		if (!OptionalCommand)
+		TOptional<TSharedRef<Private::FBuildTransferState>> OptionalState = BuildTransferThreadStates.Dequeue();
+		if (!OptionalState)
 		{
 			FPlatformProcess::Sleep(0.2f);
 			continue;
 		}
-		TSharedRef<FBuildTransferCommand> TransferCommand(OptionalCommand.GetValue());
+		TSharedRef<Private::FBuildTransferState> State(OptionalState.GetValue());
 
-		if (TransferCommand->bCancelRequested.load(std::memory_order_relaxed))
+		if (State->bCancelRequested.load(std::memory_order_relaxed))
 		{
-			FWriteScopeLock _(TransferCommand->Lock);
-			TransferCommand->Status = EBuildTransferStatus::Canceled;
+			FWriteScopeLock _(State->Lock);
+			State->Status = EBuildTransferStatus::Canceled;
 			continue;
 		}
 
  		FString ZenUtilityPath = FPaths::ConvertRelativePathToFull(UE::Zen::GetLocalInstallUtilityPath());
  		FString OidcTokenExecutableFilename = FPaths::ConvertRelativePathToFull(FDesktopPlatformModule::TryGet()->GetOidcTokenExecutableFilename(FPaths::RootDir()));
 
+ 		FString DestinationDirectory;
 		FString CommandLineArgs;
-		switch (TransferCommand->Type)
+		switch (State->Type)
 		{
-		case EBuildTransferType::Files:
-			CommandLineArgs = FString::Printf(TEXT("builds download --host %s --local-path \"%s\" --namespace %s --bucket %s --build-id %s --oidctoken-exe-path \"%s\""),
-				*WriteToString<64>(EffectiveDomain), *FPaths::ConvertRelativePathToFull(TransferCommand->Destination), *TransferCommand->Namespace, *TransferCommand->Bucket,
-				*WriteToString<64>(TransferCommand->BuildId), *OidcTokenExecutableFilename);
+		case Private::EBuildTransferType::Files:
+			DestinationDirectory = FPaths::ConvertRelativePathToFull(State->Destination);
+			CommandLineArgs = FString::Printf(TEXT("builds download --log-progress --host %s --local-path \"%s\" --namespace %s --bucket %s --build-id %s --oidctoken-exe-path \"%s\""),
+				*WriteToString<64>(EffectiveDomain), *DestinationDirectory, *State->Namespace, *State->Bucket,
+				*WriteToString<64>(State->BuildId), *OidcTokenExecutableFilename);
 			break;
-		case EBuildTransferType::Oplog:
+		case Private::EBuildTransferType::Oplog:
 			{
 				FString DestinationProjectId, DestinationOplogId;
 				FString SplitDelim(TEXT("/"));
-				TransferCommand->Destination.Split(SplitDelim, &DestinationProjectId, &DestinationOplogId);
+				State->Destination.Split(SplitDelim, &DestinationProjectId, &DestinationOplogId);
 
 				UE::Zen::FZenLocalServiceRunContext RunContext;
 				uint16 LocalPort = 8558;
@@ -705,23 +870,25 @@ FBuildServiceInstance::BuildTransferThreadLoop()
 					}
 				}
 
-				if (!TransferCommand->ProjectFilePath.IsEmpty())
+				if (!State->ProjectFilePath.IsEmpty())
 				{
-					FString OplogLifetimeMarkerPath = FPaths::Combine(FPaths::GetPath(TransferCommand->ProjectFilePath),
+					DestinationDirectory = FPaths::Combine(FPaths::GetPath(State->ProjectFilePath),
 														TEXT("Saved"),
 														TEXT("Cooked"),
-														DestinationOplogId,
+														DestinationOplogId);
+
+					FString OplogLifetimeMarkerPath = FPaths::Combine(DestinationDirectory,
 														TEXT("ue.projectstore"));
 
 					Zen::FZenProjectStoreWriter ProjectStoreWriter(*OplogLifetimeMarkerPath);
-					ProjectStoreWriter.Write(TEXT("[::1]"), LocalPort, true, DestinationProjectId, DestinationOplogId, TransferCommand->TargetPlatformName);
+					ProjectStoreWriter.Write(TEXT("[::1]"), LocalPort, true, DestinationProjectId, DestinationOplogId, State->TargetPlatformName);
 
 					FString RootDir = FPaths::RootDir();
 					FString EngineDir = FPaths::EngineDir();
 					FPaths::NormalizeDirectoryName(EngineDir);
-					FString ProjectDir = FPaths::GetPath(TransferCommand->ProjectFilePath);
+					FString ProjectDir = FPaths::GetPath(State->ProjectFilePath);
 					FPaths::NormalizeDirectoryName(ProjectDir);
-					FString ProjectPath = TransferCommand->ProjectFilePath;
+					FString ProjectPath = State->ProjectFilePath;
 					FPaths::NormalizeFilename(ProjectPath);
 
 					IPlatformFile& PlatformFile = FPlatformFileManager::Get().GetPlatformFile();
@@ -747,8 +914,8 @@ FBuildServiceInstance::BuildTransferThreadLoop()
 
 				CommandLineArgs = FString::Printf(TEXT("oplog-import %s %s --clean --builds %s --namespace %s --bucket %s --builds-id %s --access-token \"%s\""),
 					*DestinationProjectId, *DestinationOplogId, *WriteToString<64>(EffectiveDomain),
-					*TransferCommand->Namespace, *TransferCommand->Bucket,
-					*WriteToString<64>(TransferCommand->BuildId), *GetAccessToken());
+					*State->Namespace, *State->Bucket,
+					*WriteToString<64>(State->BuildId), *GetAccessToken());
 
 
 			}
@@ -756,22 +923,70 @@ FBuildServiceInstance::BuildTransferThreadLoop()
 		}
 
 		FMonitoredProcess MonitoredUtilityProcess(ZenUtilityPath, *CommandLineArgs, FPaths::GetPath(ZenUtilityPath), true);
-		MonitoredUtilityProcess.OnOutput().BindSPLambda(TransferCommand, [TransferCommand](FString Output)
+		MonitoredUtilityProcess.OnOutput().BindSPLambda(State, [State](FString Output)
 		{
-			FWriteScopeLock _(TransferCommand->Lock);
-			TransferCommand->RecentOutputLines[TransferCommand->NextOutputLineIndex++] = MoveTemp(Output);
+			FReadScopeLock _(State->Lock);
+			if (Output.StartsWith(TEXT("@progress push")))
+			{
+				Output = Output.RightChop(15);
+				Private::FProgressState NewState;
+				NewState.Label = Output;
+				State->ProgressState.Push(MoveTemp(NewState));
+				return;
+			}
+			else if (Output.StartsWith(TEXT("@progress pop")))
+			{
+				State->ProgressState.Pop();
+				return;
+			}
+			else if (Output.StartsWith(TEXT("@progress")))
+			{
+				if (State->ProgressState.IsEmpty())
+				{
+					Private::FProgressState NewState;
+					State->ProgressState.Push(MoveTemp(NewState));
+				}
+				Private::FProgressState& CurrentState = State->ProgressState.Top();
+				Output = Output.RightChop(10);
+				if (Output.EndsWith(TEXT("%")))
+				{
+					FString PossiblePercentText = Output.LeftChop(1);
+					if (!PossiblePercentText.IsEmpty())
+					{
+						float Percent = FCString::Atof(*PossiblePercentText);
+						if ((Percent != 0.0f) || (PossiblePercentText.Len() == 1 && PossiblePercentText[0] == TCHAR('0')))
+						{
+							CurrentState.Percent = Percent;
+							return; // Do not print the percent progress alone
+						}
+					}
+				}
+				CurrentState.Detail = Output;
+				Output = FString::Printf(TEXT("[%s %d%%] %s"), *CurrentState.Label, (int32)CurrentState.Percent, *Output);
+				// Deliberate fall-through so that the text is printed;
+			}
+			if (State->OutputLogFile)
+			{
+				State->OutputLogFile->Logf(TEXT("%s"), *Output);
+				State->OutputLogFile->Flush();
+			}
+			State->RecentOutputLines[State->NextOutputLineIndex++] = MoveTemp(Output);
 		});
 
 		{
-			FWriteScopeLock _(TransferCommand->Lock);
-			TransferCommand->Status = EBuildTransferStatus::Active;
+			FWriteScopeLock _(State->Lock);
+			FString LogFilename = State->GetLogFilename();
+			FOutputDeviceFile::CreateBackupCopy(*LogFilename);
+			State->OutputLogFile = TUniquePtr<FArchive>(IFileManager::Get().CreateFileWriter(*LogFilename, FILEWRITE_Silent | FILEWRITE_AllowRead));
+			State->Status = EBuildTransferStatus::Active;
 		}
 
 		if (!MonitoredUtilityProcess.Launch())
 		{
 			UE_LOG(LogBuildServiceInstance, Warning, TEXT("Failed to launch zen utility to download build: '%s'."), *ZenUtilityPath);
-			FWriteScopeLock _(TransferCommand->Lock);
-			TransferCommand->Status = EBuildTransferStatus::Failed;
+			FWriteScopeLock _(State->Lock);
+			State->Status = EBuildTransferStatus::Failed;
+			State->OutputLogFile.Reset();
 			continue;
 		}
 
@@ -783,7 +998,7 @@ FBuildServiceInstance::BuildTransferThreadLoop()
 				return;
 			}
 			double Duration = FPlatformTime::ToSeconds64(FPlatformTime::Cycles64() - StartTime);
-			if (bBuildTransferThreadStopping.load(std::memory_order_relaxed) || TransferCommand->bCancelRequested.load(std::memory_order_relaxed))
+			if (bBuildTransferThreadStopping.load(std::memory_order_relaxed) || State->bCancelRequested.load(std::memory_order_relaxed))
 			{
 				MonitoredUtilityProcess.Cancel(true);
 				UE_LOG(LogBuildServiceInstance, Display, TEXT("Canceled launch of zen utility for downloading build data: '%s'."), *ZenUtilityPath);
@@ -805,29 +1020,61 @@ FBuildServiceInstance::BuildTransferThreadLoop()
 				}
 
 				UE_LOG(LogBuildServiceInstance, Display, TEXT("Launch of zen utility for downloading build data: '%s' canceled."), *ZenUtilityPath);
-				FWriteScopeLock _(TransferCommand->Lock);
-				TransferCommand->Status = EBuildTransferStatus::Canceled;
+				FWriteScopeLock _(State->Lock);
+				State->Status = EBuildTransferStatus::Canceled;
+				State->OutputLogFile.Reset();
 				break;
 			}
 			FPlatformProcess::Sleep(0.1f);
 		}
 
-		if (!TransferCommand->bCancelRequested.load(std::memory_order_relaxed))
+		if (!State->bCancelRequested.load(std::memory_order_relaxed))
 		{
 			if (MonitoredUtilityProcess.GetReturnCode() == 0)
 			{
-				FWriteScopeLock _(TransferCommand->Lock);
-				TransferCommand->ReturnCode = MonitoredUtilityProcess.GetReturnCode();
-				TransferCommand->Status = EBuildTransferStatus::Succeeded;
+				FWriteScopeLock _(State->Lock);
+				State->ReturnCode = MonitoredUtilityProcess.GetReturnCode();
+				State->Status = EBuildTransferStatus::Succeeded;
+				State->OutputLogFile.Reset();
 			}
 			else
 			{
-				UE_LOG(LogBuildServiceInstance, Warning, TEXT("Unexpected return code after launch of zen utility for downloading build data: '%s' (%d). Output: '%s'"), *ZenUtilityPath, MonitoredUtilityProcess.GetReturnCode(), *TransferCommand->Output);
-				FWriteScopeLock _(TransferCommand->Lock);
-				TransferCommand->ReturnCode = MonitoredUtilityProcess.GetReturnCode();
-				TransferCommand->Status = EBuildTransferStatus::Failed;
+				FWriteScopeLock _(State->Lock);
+				UE_LOG(LogBuildServiceInstance, Warning, TEXT("Unexpected return code after launch of zen utility for downloading build data: '%s' (%d). Output: '%s'"), *ZenUtilityPath, MonitoredUtilityProcess.GetReturnCode(), *State->GetRecentOutput());
+				State->ReturnCode = MonitoredUtilityProcess.GetReturnCode();
+				State->Status = EBuildTransferStatus::Failed;
+				State->OutputLogFile.Reset();
 			}
 		}
+
+		bool bCleanDestination = false;
+		{
+			FReadScopeLock _(State->Lock);
+			bCleanDestination = (State->Status == EBuildTransferStatus::Succeeded) &&
+				(State->Type == Private::EBuildTransferType::Oplog) &&
+				!DestinationDirectory.IsEmpty();
+		}
+
+		if (bCleanDestination)
+		{
+			IPlatformFile& PlatformFile = FPlatformFileManager::Get().GetPlatformFile();
+			PlatformFile.IterateDirectory(*DestinationDirectory,
+				[&PlatformFile](const TCHAR* FoundFullPath, bool bDirectory)
+				{
+					if (bDirectory)
+					{
+						PlatformFile.DeleteDirectoryRecursively(FoundFullPath);
+					}
+					else
+					{
+						if (FPathViews::GetCleanFilename(FoundFullPath) != TEXT("ue.projectstore"))
+						{
+							PlatformFile.DeleteFile(FoundFullPath);
+						}
+					}
+					return true;
+				});
+		}
 	}
 }
 
diff --git a/Engine/Source/Developer/Zen/Public/Experimental/BuildServerInterface.h b/Engine/Source/Developer/Zen/Public/Experimental/BuildServerInterface.h
--- a/Engine/Source/Developer/Zen/Public/Experimental/BuildServerInterface.h
+++ b/Engine/Source/Developer/Zen/Public/Experimental/BuildServerInterface.h
@@ -15,12 +15,14 @@
 #include "Serialization/CompactBinary.h"
 #include "Templates/Function.h"
 #include "Templates/SharedPointer.h"
+#include "Templates/UniquePtr.h"
 #include <atomic>
 
 #define UE_API ZEN_API
 
 namespace UE { class IHttpClient; }
 namespace UE::Zen::Build::Private { class FAccessToken; }
+namespace UE::Zen::Build::Private { struct FBuildTransferState; }
 
 namespace UE::Zen::Build
 {
@@ -53,35 +55,6 @@ public:
 		Succeeded
 	};
 
-private:
-	enum class EBuildTransferType
-	{
-		Files,
-		Oplog
-	};
-
-	struct FBuildTransferCommand : public TSharedFromThis<FBuildTransferCommand, ESPMode::ThreadSafe>
-	{
-		EBuildTransferType Type;
-		FCbObjectId BuildId;
-		FString Namespace;
-		FString Bucket;
-		FString Destination;
-		FString ProjectFilePath;
-		FName TargetPlatformName = NAME_None;
-
-		std::atomic<bool> bCancelRequested = false;
-
-		FRWLock Lock;
-		const uint32 RecentOutputLinesCapacity = 16;
-		TCircularBuffer<FString> RecentOutputLines = TCircularBuffer<FString>(RecentOutputLinesCapacity);
-		uint32 NextOutputLineIndex = 0;
-		FString Output;
-		EBuildTransferStatus Status = EBuildTransferStatus::Queued;
-		int ReturnCode;
-	};
-
-public:
 	enum class EConnectionState
 	{
 		NotStarted,
@@ -100,11 +73,16 @@ public:
 	{
 	public:
 		UE_API FString GetDescription() const;
+		UE_API FString GetRecentOutput() const;
 		UE_API EBuildTransferStatus GetStatus() const;
+		UE_API FString GetLogFilename() const;
+
+		UE_API bool GetCurrentProgress(FString& OutLabel, FString& OutDetail, float& OutPercent) const;
+		UE_API bool GetOverallProgress(FString& OutLabel, FString& OutDetail, float& OutPercent) const;
 
 		UE_API void RequestCancel();
 	private:
-		TSharedPtr<FBuildTransferCommand> BuildTransferCommand;
+		TSharedPtr<Private::FBuildTransferState> State;
 		friend class UE::Zen::Build::FBuildServiceInstance;
 	};
 
@@ -129,8 +107,8 @@ public:
 
 	UE_API void ListBuilds(FStringView Namespace, FStringView Bucket, FOnListBuildsComplete&& InOnListBuildsComplete);
 
-	UE_API FBuildTransfer StartBuildTransfer(const FCbObjectId& BuildId, FStringView DestinationFolder, FStringView Namespace, FStringView Bucket);
-	UE_API FBuildTransfer StartOplogBuildTransfer(const FCbObjectId& BuildId, FStringView DestinationProjectId, FStringView DestinationOplogId, FStringView ProjectFilePath, FStringView Namespace, FStringView Bucket, FName TargetPlatformName);
+	UE_API FBuildTransfer StartBuildTransfer(const FCbObjectId& BuildId, FStringView Name, FStringView DestinationFolder, FStringView Namespace, FStringView Bucket);
+	UE_API FBuildTransfer StartOplogBuildTransfer(const FCbObjectId& BuildId, FStringView Name, FStringView DestinationProjectId, FStringView DestinationOplogId, FStringView ProjectFilePath, FStringView Namespace, FStringView Bucket, FName TargetPlatformName);
 private:
 	typedef TUniqueFunction<void(FString& CommandLineArgs)> FPreZenUtilityInvocation;
 	typedef TUniqueFunction<void(bool)> FOnZenUtilityInvocationComplete;
@@ -161,7 +139,7 @@ private:
 	FRWLock NamespacesAndBucketsLock;
 	TMultiMap<FString, FString> NamespacesAndBuckets;
 
-	TMpscQueue<TSharedRef<FBuildTransferCommand>> BuildTransferThreadCommands;
+	TMpscQueue<TSharedRef<Private::FBuildTransferState>> BuildTransferThreadStates;
 	FThread BuildTransferThread;
 	std::atomic<bool> bBuildTransferThreadStarting = false;
 	std::atomic<bool> bBuildTransferThreadStopping = false;
diff --git a/Engine/Source/Editor/StorageServerWidgets/Internal/SBuildActivity.h b/Engine/Source/Editor/StorageServerWidgets/Internal/SBuildActivity.h
--- a/Engine/Source/Editor/StorageServerWidgets/Internal/SBuildActivity.h
+++ b/Engine/Source/Editor/StorageServerWidgets/Internal/SBuildActivity.h
@@ -53,7 +53,11 @@ private:
 
 	UE_API bool CanCancelBuildActivity() const;
 	UE_API void CancelBuildActivity() const;
-	UE_API void ClearCompleted();
+	UE_API bool CanClearBuildActivity() const;
+	UE_API void ClearBuildActivity();
+	UE_API bool CanViewLogForBuildActivity() const;
+	UE_API void ViewLogForBuildActivity() const;
+	UE_API void ClearAllCompleted();
 
 	TSharedPtr<SListView<TSharedPtr<SBuildActivity::FBuildActivity, ESPMode::ThreadSafe>>> BuildActivityListView;
 	TArray<TSharedPtr<SBuildActivity::FBuildActivity, ESPMode::ThreadSafe>> BuildActivities;
diff --git a/Engine/Source/Editor/StorageServerWidgets/Internal/SBuildSelection.h b/Engine/Source/Editor/StorageServerWidgets/Internal/SBuildSelection.h
--- a/Engine/Source/Editor/StorageServerWidgets/Internal/SBuildSelection.h
+++ b/Engine/Source/Editor/StorageServerWidgets/Internal/SBuildSelection.h
@@ -67,6 +67,7 @@ private:
 	UE_API EBuildType GetSelectedBuildType() const;
 	UE_API void RebuildLists();
 	UE_API void RegenerateBuildGroups(UE::BuildSelection::Internal::FListBuildsState& ListBuildsState);
+	UE_API void ConditionalSortBuildGroups();
 	UE_API void RegenerateActivePlatformFilters();
 	UE_API void ValidateBuildGroupSelection();
 	UE_API TSharedRef<SWidget> OnGenerateTextBlockFromString(TSharedPtr<FString> Item);
@@ -77,6 +78,9 @@ private:
 	UE_API TSharedRef<SWidget> GetBuildDestinationPanel();
 	UE_API TSharedRef<SWidget> GetGridPanel();
 	UE_API void BuildGroupSelectionChanged(TSharedPtr<FBuildGroup> Item, ESelectInfo::Type SelectInfo);
+	EColumnSortMode::Type GetColumnSortMode(const FName ColumnId) const;
+	void OnColumnSortModeChanged(const EColumnSortPriority::Type SortPriority, const FName& ColumnId, const EColumnSortMode::Type InSortMode);
+
 
 	UE_API FString& GetUserSelectedDestination();
 	UE_API FString GetUserSelectedDestination() const;
@@ -120,6 +124,8 @@ private:
 	TSharedPtr<FString> SelectedBuildType;
 
 	TSharedPtr<SListView<TSharedPtr<SBuildSelection::FBuildGroup, ESPMode::ThreadSafe>>> BuildGroupListView;
+	FName BuildGroupSortByColumn;
+	EColumnSortMode::Type BuildGroupSortMode;
 
 	std::atomic<uint32> BuildListRefreshesInProgress = 0;
 	std::atomic<uint32> BuildRefreshGeneration = 0;
diff --git a/Engine/Source/Editor/StorageServerWidgets/Private/SBuildActivity.cpp b/Engine/Source/Editor/StorageServerWidgets/Private/SBuildActivity.cpp
--- a/Engine/Source/Editor/StorageServerWidgets/Private/SBuildActivity.cpp
+++ b/Engine/Source/Editor/StorageServerWidgets/Private/SBuildActivity.cpp
@@ -11,6 +11,7 @@
 #include "Widgets/Input/SButton.h"
 #include "Widgets/Input/SHyperlink.h"
 #include "Widgets/Layout/SGridPanel.h"
+#include "Widgets/Notifications/SProgressBar.h"
 #include "Widgets/SBoxPanel.h"
 #include "Widgets/Text/STextBlock.h"
 #include "Widgets/Views/SListView.h"
@@ -97,7 +98,7 @@ TSharedRef<SWidget> SBuildActivity::GetGridPanel()
 				+ SHeaderRow::Column(FBuildActivityIds::ColPlatform).DefaultLabel(LOCTEXT("BuildActivity_ActivityPlatform", "Platform"))
 					.FillWidth(0.1f).HAlignCell(HAlign_Left).HAlignHeader(HAlign_Left).VAlignCell(VAlign_Center)
 				+ SHeaderRow::Column(FBuildActivityIds::ColDescription).DefaultLabel(LOCTEXT("BuildActivity_ActivityDescription", "Description"))
-					.FillWidth(0.6f).HAlignCell(HAlign_Left).HAlignHeader(HAlign_Center).VAlignCell(VAlign_Center)
+					.FillWidth(0.6f).HAlignCell(HAlign_Fill).HAlignHeader(HAlign_Center).VAlignCell(VAlign_Center)
 			)
 	];
 
@@ -127,9 +128,21 @@ TSharedPtr<SWidget> SBuildActivity::OnGetBuildActivityContextMenuContent()
 		NAME_None,
 		bRecursivelySearchable);
 
+	MenuBuilder.AddMenuEntry(
+		LOCTEXT("BuildActivity_ViewLog", "View log"),
+		LOCTEXT("BuildActivity_ViewLogTooltip", "View log file for the activity"),
+		FSlateIcon(),
+		FUIAction(
+			FExecuteAction::CreateRaw(this, &SBuildActivity::ViewLogForBuildActivity),
+			FCanExecuteAction::CreateRaw(this, &SBuildActivity::CanViewLogForBuildActivity)
+		),
+		NAME_None,
+		EUserInterfaceActionType::Button
+	);
+
 	MenuBuilder.AddMenuEntry(
 		LOCTEXT("BuildActivity_Cancel", "Cancel"),
-		LOCTEXT("BuildActivity_CancelTooltip", "Cancel the pending activity"),
+		LOCTEXT("BuildActivity_CancelTooltip", "Cancel the pending or in progress activity"),
 		FSlateIcon(),
 		FUIAction(
 			FExecuteAction::CreateRaw(this, &SBuildActivity::CancelBuildActivity),
@@ -140,11 +153,23 @@ TSharedPtr<SWidget> SBuildActivity::OnGetBuildActivityContextMenuContent()
 	);
 
 	MenuBuilder.AddMenuEntry(
-		LOCTEXT("BuildActivity_ClearCompleted", "Clear completed"),
-		LOCTEXT("BuildActivity_ClearCompletedTooltip", "Clear all completed items"),
+		LOCTEXT("BuildActivity_Clear", "Clear"),
+		LOCTEXT("BuildActivity_ClearTooltip", "Clear this activity if it is completed"),
 		FSlateIcon(),
 		FUIAction(
-			FExecuteAction::CreateRaw(this, &SBuildActivity::ClearCompleted)
+			FExecuteAction::CreateRaw(this, &SBuildActivity::ClearBuildActivity),
+			FCanExecuteAction::CreateRaw(this, &SBuildActivity::CanClearBuildActivity)
+		),
+		NAME_None,
+		EUserInterfaceActionType::Button
+	);
+
+	MenuBuilder.AddMenuEntry(
+		LOCTEXT("BuildActivity_ClearAllCompleted", "Clear all completed"),
+		LOCTEXT("BuildActivity_ClearAllCompletedTooltip", "Clear all completed items"),
+		FSlateIcon(),
+		FUIAction(
+			FExecuteAction::CreateRaw(this, &SBuildActivity::ClearAllCompleted)
 		),
 		NAME_None,
 		EUserInterfaceActionType::Button
@@ -178,7 +203,61 @@ void SBuildActivity::CancelBuildActivity() const
 	SelectedItems[0]->Transfer.RequestCancel();
 }
 
-void SBuildActivity::ClearCompleted()
+bool SBuildActivity::CanClearBuildActivity() const
+{
+	using namespace UE::Zen::Build;
+
+	TArray<FBuildActivityPtr> SelectedItems = BuildActivityListView->GetSelectedItems();
+	if (SelectedItems.Num() != 1)
+	{
+		return false;
+	}
+	FBuildServiceInstance::EBuildTransferStatus Status = SelectedItems[0]->Transfer.GetStatus();
+	return (Status != FBuildServiceInstance::EBuildTransferStatus::Queued) && (Status != FBuildServiceInstance::EBuildTransferStatus::Active);
+}
+
+void SBuildActivity::ClearBuildActivity()
+{
+	using namespace UE::Zen::Build;
+
+	TArray<FBuildActivityPtr> SelectedItems = BuildActivityListView->GetSelectedItems();
+	if (SelectedItems.Num() != 1)
+	{
+		return;
+	}
+	BuildActivities.Remove(SelectedItems[0]);
+}
+
+bool SBuildActivity::CanViewLogForBuildActivity() const
+{
+	using namespace UE::Zen::Build;
+
+	TArray<FBuildActivityPtr> SelectedItems = BuildActivityListView->GetSelectedItems();
+	if (SelectedItems.Num() != 1)
+	{
+		return false;
+	}
+	FString LogFilename = SelectedItems[0]->Transfer.GetLogFilename();
+	return !LogFilename.IsEmpty() && FPaths::FileExists(LogFilename);
+}
+
+void SBuildActivity::ViewLogForBuildActivity() const
+{
+	using namespace UE::Zen::Build;
+
+	TArray<FBuildActivityPtr> SelectedItems = BuildActivityListView->GetSelectedItems();
+	if (SelectedItems.Num() != 1)
+	{
+		return;
+	}
+	FString LogFilename = SelectedItems[0]->Transfer.GetLogFilename();
+	if (!LogFilename.IsEmpty() && FPaths::FileExists(LogFilename))
+	{
+		FPlatformProcess::LaunchFileInDefaultExternalApplication(*LogFilename);
+	}
+}
+
+void SBuildActivity::ClearAllCompleted()
 {
 	using namespace UE::Zen::Build;
 
@@ -222,8 +301,11 @@ SBuildActivityTableRow::GenerateWidgetForColumn(const FName& ColumnName)
 	{
 		return SNew(SHorizontalBox)
 				+ SHorizontalBox::Slot()
+				.Padding(1.f,1.f)
 				[
-					SNew(STextBlock).Text_Lambda([this]()
+					SNew(STextBlock)
+					.TextStyle(&FAppStyle::GetWidgetStyle<FTextBlockStyle>(TEXT("SmallText")))
+					.Text_Lambda([this]()
 						{
 							FString Status;
 							if (BuildServiceInstance)
@@ -252,33 +334,81 @@ SBuildActivityTableRow::GenerateWidgetForColumn(const FName& ColumnName)
 	{
 		return SNew(SHorizontalBox)
 				+ SHorizontalBox::Slot()
+				.Padding(1.f,1.f)
 				[
-					SNew(STextBlock).Text(FText::FromString(BuildActivity->Name))
+					SNew(STextBlock)
+					.TextStyle(&FAppStyle::GetWidgetStyle<FTextBlockStyle>(TEXT("SmallText")))
+					.Text(FText::FromString(BuildActivity->Name))
 				];
 	}
 	else if (ColumnName == FBuildActivityIds::ColPlatform)
 	{
 		return SNew(SHorizontalBox)
 				+ SHorizontalBox::Slot()
+				.Padding(1.f,1.f)
 				[
-					SNew(STextBlock).Text(FText::FromString(BuildActivity->Platform))
+					SNew(STextBlock)
+					.TextStyle(&FAppStyle::GetWidgetStyle<FTextBlockStyle>(TEXT("SmallText")))
+					.Text(FText::FromString(BuildActivity->Platform))
 				];
 	}
 	else if (ColumnName == FBuildActivityIds::ColDescription)
 	{
-		return SNew(SHorizontalBox)
-				+ SHorizontalBox::Slot()
-				[
-					SNew(STextBlock).Text_Lambda([this]()
+		return SNew(SOverlay)
+					+ SOverlay::Slot()
+					.Padding(1.f,1.f)
+					[
+						SNew(SProgressBar)
+						.Visibility_Lambda([this]()
 						{
-							FString Description;
-							if (BuildServiceInstance)
+							return BuildActivity->Transfer.GetStatus() == FBuildServiceInstance::EBuildTransferStatus::Active ? EVisibility::Visible : EVisibility::Collapsed;
+						})
+						.Percent_Lambda([this]()
+						{
+							FString Label;
+							FString Detail;
+							float Percent;
+							if (BuildActivity->Transfer.GetCurrentProgress(Label, Detail, Percent))
 							{
-								Description = BuildActivity->Transfer.GetDescription();
+								return Percent/100.f;
 							}
-							return FText::FromString(Description);
+							return 0.f;
 						})
-				];
+					]
+					+ SOverlay::Slot()
+					.Padding(1.f,1.f)
+					.VAlign(VAlign_Center)
+					[
+						SNew(STextBlock)
+						.Text_Lambda([this]()
+						{
+							FString Label;
+							FString Detail;
+							float Percent;
+							if (BuildActivity->Transfer.GetCurrentProgress(Label, Detail, Percent))
+							{
+								return FText::FromString(Label);
+							}
+							return FText::FromString(FString());
+						})
+						.TextStyle(&FAppStyle::GetWidgetStyle<FTextBlockStyle>(TEXT("SmallText")))
+						.Justification(ETextJustify::Type::Left)
+						.ToolTipText_Lambda([this]()
+						{
+							FString Label;
+							FString Detail;
+							float Percent;
+							if (BuildActivity->Transfer.GetCurrentProgress(Label, Detail, Percent))
+							{
+								return FText::FromString(Detail);
+							}
+							return FText::FromString(FString());
+						})
+						.Visibility_Lambda([this]()
+						{
+							return BuildActivity->Transfer.GetStatus() == FBuildServiceInstance::EBuildTransferStatus::Active ? EVisibility::Visible : EVisibility::Collapsed;
+						})
+					];
 	}
 
 	return SNullWidget::NullWidget;
diff --git a/Engine/Source/Editor/StorageServerWidgets/Private/SBuildSelection.cpp b/Engine/Source/Editor/StorageServerWidgets/Private/SBuildSelection.cpp
--- a/Engine/Source/Editor/StorageServerWidgets/Private/SBuildSelection.cpp
+++ b/Engine/Source/Editor/StorageServerWidgets/Private/SBuildSelection.cpp
@@ -65,10 +65,15 @@ namespace UE::BuildSelection::Internal
 
 void SBuildSelection::Construct(const FArguments& InArgs)
 {
+	using namespace UE::BuildSelection::Internal;
+
 	ZenServiceInstance = InArgs._ZenServiceInstance;
 	BuildServiceInstance = InArgs._BuildServiceInstance;
 	OnBuildTransferStarted = InArgs._OnBuildTransferStarted;
 
+	BuildGroupSortByColumn = FBuildGroupIds::ColCreated;
+	BuildGroupSortMode = EColumnSortMode::Descending;
+
 	FBuildVersion CurrentBuildVersion;
 	if (FBuildVersion::TryRead(FBuildVersion::GetDefaultFileName(), CurrentBuildVersion))
 	{
@@ -536,6 +541,85 @@ void SBuildSelection::RegenerateBuildGroups(UE::BuildSelection::Internal::FListB
 	}
 
 	KeyedGroups.GenerateValueArray(BuildGroups);
+	ConditionalSortBuildGroups();
+}
+
+void SBuildSelection::ConditionalSortBuildGroups()
+{
+	using namespace UE::BuildSelection::Internal;
+
+	if (BuildGroupSortMode == EColumnSortMode::None)
+	{
+		return;
+	}
+
+	// Sorting
+	if (BuildGroupSortMode == EColumnSortMode::Ascending)
+	{
+		if (BuildGroupSortByColumn == FBuildGroupIds::ColName)
+		{
+			BuildGroups.Sort([this](const TSharedPtr<FBuildGroup>& A, const TSharedPtr<FBuildGroup>& B) {
+				return (A->DisplayName < B->DisplayName);
+			});
+		}
+		else if (BuildGroupSortByColumn == FBuildGroupIds::ColCommit)
+		{
+			BuildGroups.Sort([this](const TSharedPtr<FBuildGroup>& A, const TSharedPtr<FBuildGroup>& B) {
+				return (A->CommitIdentifier < B->CommitIdentifier);
+			});
+		}
+		else if (BuildGroupSortByColumn == FBuildGroupIds::ColSuffix)
+		{
+			BuildGroups.Sort([this](const TSharedPtr<FBuildGroup>& A, const TSharedPtr<FBuildGroup>& B) {
+				return (A->Suffix < B->Suffix);
+			});
+		}
+		else if (BuildGroupSortByColumn == FBuildGroupIds::ColCategory)
+		{
+			BuildGroups.Sort([this](const TSharedPtr<FBuildGroup>& A, const TSharedPtr<FBuildGroup>& B) {
+				return (A->Category < B->Category);
+			});
+		}
+		else if (BuildGroupSortByColumn == FBuildGroupIds::ColCreated)
+		{
+			BuildGroups.Sort([this](const TSharedPtr<FBuildGroup>& A, const TSharedPtr<FBuildGroup>& B) {
+				return (A->CreatedAt < B->CreatedAt);
+			});
+		}
+	}
+	else
+	{
+		if (BuildGroupSortByColumn == FBuildGroupIds::ColName)
+		{
+			BuildGroups.Sort([this](const TSharedPtr<FBuildGroup>& A, const TSharedPtr<FBuildGroup>& B) {
+				return (A->DisplayName >= B->DisplayName);
+			});
+		}
+		else if (BuildGroupSortByColumn == FBuildGroupIds::ColCommit)
+		{
+			BuildGroups.Sort([this](const TSharedPtr<FBuildGroup>& A, const TSharedPtr<FBuildGroup>& B) {
+				return (A->CommitIdentifier >= B->CommitIdentifier);
+			});
+		}
+		else if (BuildGroupSortByColumn == FBuildGroupIds::ColSuffix)
+		{
+			BuildGroups.Sort([this](const TSharedPtr<FBuildGroup>& A, const TSharedPtr<FBuildGroup>& B) {
+				return (A->Suffix >= B->Suffix);
+			});
+		}
+		else if (BuildGroupSortByColumn == FBuildGroupIds::ColCategory)
+		{
+			BuildGroups.Sort([this](const TSharedPtr<FBuildGroup>& A, const TSharedPtr<FBuildGroup>& B) {
+				return (A->Category >= B->Category);
+			});
+		}
+		else if (BuildGroupSortByColumn == FBuildGroupIds::ColCreated)
+		{
+			BuildGroups.Sort([this](const TSharedPtr<FBuildGroup>& A, const TSharedPtr<FBuildGroup>& B) {
+				return (A->CreatedAt >= B->CreatedAt);
+			});
+		}
+	}
 }
 
 void SBuildSelection::RegenerateActivePlatformFilters()
@@ -791,6 +875,29 @@ void SBuildSelection::BuildGroupSelectionChanged(FBuildSelectionBuildGroupPtr It
 	}
 }
 
+EColumnSortMode::Type SBuildSelection::GetColumnSortMode(const FName ColumnId) const
+{
+	if (BuildGroupSortByColumn == ColumnId)
+	{
+		return BuildGroupSortMode;
+	}
+
+	return EColumnSortMode::None;
+}
+
+void SBuildSelection::OnColumnSortModeChanged(const EColumnSortPriority::Type SortPriority, const FName& ColumnId, const EColumnSortMode::Type InSortMode)
+{
+	BuildGroupSortByColumn = ColumnId;
+	BuildGroupSortMode = InSortMode;
+
+	ConditionalSortBuildGroups();
+	ExecuteOnGameThread(UE_SOURCE_LOCATION,
+		[this]
+		{
+			BuildGroupListView->RequestListRefresh();
+		});
+}
+
 FString& SBuildSelection::GetUserSelectedDestination()
 {
 	EBuildType BuildType = GetSelectedBuildType();
@@ -1348,18 +1455,28 @@ TSharedRef<SWidget> SBuildSelection::GetGridPanel()
 					SNew(SHeaderRow)
 					+ SHeaderRow::Column(FBuildGroupIds::ColName).DefaultLabel(LOCTEXT("BuildSelection_BuildGroupColName", "Name"))
 						.FillWidth(0.4f)
+						.SortMode(this, &SBuildSelection::GetColumnSortMode, FBuildGroupIds::ColName)
+						.OnSort(this, &SBuildSelection::OnColumnSortModeChanged)
 					+ SHeaderRow::Column(FBuildGroupIds::ColCommit).DefaultLabel(LOCTEXT("BuildSelection_BuildGroupColCommit", "Commit"))
 						.DefaultTooltip(LOCTEXT("BuildSelection_BuildGroupColCommitTooltip", "Commit/Changelist for the build"))
 						.FillWidth(0.10f).HAlignCell(HAlign_Center).HAlignHeader(HAlign_Center).VAlignCell(VAlign_Center)
+						.SortMode(this, &SBuildSelection::GetColumnSortMode, FBuildGroupIds::ColCommit)
+						.OnSort(this, &SBuildSelection::OnColumnSortModeChanged)
 					+ SHeaderRow::Column(FBuildGroupIds::ColSuffix).DefaultLabel(LOCTEXT("BuildSelection_BuildGroupColSuffix", "Suffix"))
 						.DefaultTooltip(LOCTEXT("BuildSelection_BuildGroupColSuffixTooltip", "Modifier on top of the commit/changelist for the build"))
 						.FillWidth(0.10f).HAlignCell(HAlign_Center).HAlignHeader(HAlign_Center).VAlignCell(VAlign_Center)
+						.SortMode(this, &SBuildSelection::GetColumnSortMode, FBuildGroupIds::ColSuffix)
+						.OnSort(this, &SBuildSelection::OnColumnSortModeChanged)
 					+ SHeaderRow::Column(FBuildGroupIds::ColCategory).DefaultLabel(LOCTEXT("BuildSelection_BuildGroupColCategory", "Category"))
 						.DefaultTooltip(LOCTEXT("BuildSelection_BuildGroupColCategoryTooltip", "Category for the build"))
+						.SortMode(this, &SBuildSelection::GetColumnSortMode, FBuildGroupIds::ColCategory)
+						.OnSort(this, &SBuildSelection::OnColumnSortModeChanged)
 						.FillWidth(0.25f).HAlignCell(HAlign_Left).HAlignHeader(HAlign_Center).VAlignCell(VAlign_Center)
 					+ SHeaderRow::Column(FBuildGroupIds::ColCreated).DefaultLabel(LOCTEXT("BuildSelection_BuildGroupColCreated", "Created"))
 						.DefaultTooltip(LOCTEXT("BuildSelection_BuildGroupColCreatedTooltip", "When the build was created"))
 						.FillWidth(0.15f).HAlignCell(HAlign_Left).HAlignHeader(HAlign_Center).VAlignCell(VAlign_Center)
+						.SortMode(this, &SBuildSelection::GetColumnSortMode, FBuildGroupIds::ColCreated)
+						.OnSort(this, &SBuildSelection::OnColumnSortModeChanged)
 				)
 		]
 	];
@@ -1552,6 +1669,8 @@ TSharedRef<SWidget> SBuildSelection::GetGridPanel()
 								int BuildTypeIndex = (int)BuildType;
 								const bool bAppendBuildNameToDestination = bAppendBuildNameToDestinations[BuildTypeIndex];
 
+								FString TransferName = FString::Printf(TEXT("%s-%s"), *SelectedItems[0]->DisplayName, *DestinationPlatformName);
+
 								if (BuildType == EBuildType::Oplog)
 								{
 									FString ProjectFilePath = FUProjectDictionary::GetDefault().GetProjectPathForGame(**SelectedProject);
@@ -1562,7 +1681,7 @@ TSharedRef<SWidget> SBuildSelection::GetGridPanel()
 										DestinationOplogId = FString::Printf(TEXT("%s.%s"), *DestinationOplogId, *SanitizeForZenId(SelectedItems[0]->DisplayName));
 									}
 									FBuildServiceInstance::FBuildTransfer BuildTransfer =
-										ServiceInstance->StartOplogBuildTransfer(BuildRecord->BuildId, DestinationProjectId, DestinationOplogId, ProjectFilePath, SelectedItems[0]->Namespace, Bucket, TargetPlatformName);
+										ServiceInstance->StartOplogBuildTransfer(BuildRecord->BuildId, TransferName, DestinationProjectId, DestinationOplogId, ProjectFilePath, SelectedItems[0]->Namespace, Bucket, TargetPlatformName);
 									OnBuildTransferStarted.ExecuteIfBound(BuildTransfer, SelectedItems[0]->DisplayName, PlatformForSelectedGroup);
 								}
 								else
@@ -1577,7 +1696,7 @@ TSharedRef<SWidget> SBuildSelection::GetGridPanel()
 										DestinationFolder = FPaths::Combine(GetEffectiveDestination(), DestinationPlatformName);
 									}
 									FBuildServiceInstance::FBuildTransfer BuildTransfer =
-										ServiceInstance->StartBuildTransfer(BuildRecord->BuildId, DestinationFolder, SelectedItems[0]->Namespace, Bucket);
+										ServiceInstance->StartBuildTransfer(BuildRecord->BuildId, TransferName, DestinationFolder, SelectedItems[0]->Namespace, Bucket);
 									OnBuildTransferStarted.ExecuteIfBound(BuildTransfer, SelectedItems[0]->DisplayName, PlatformForSelectedGroup);
 								}
 							}
@@ -1609,6 +1728,7 @@ SBuildGroupTableRow::GenerateWidgetForColumn(const FName& ColumnName)
 	{
 		return SNew(SHorizontalBox)
 				+ SHorizontalBox::Slot()
+				.Padding(1.f,1.f)
 				[
 					SNew(STextBlock).Text(FText::FromString(BuildGroup->DisplayName))
 				];
@@ -1617,6 +1737,7 @@ SBuildGroupTableRow::GenerateWidgetForColumn(const FName& ColumnName)
 	{
 		return SNew(SHorizontalBox)
 				+ SHorizontalBox::Slot()
+				.Padding(1.f,1.f)
 				[
 					SNew(STextBlock).Text(FText::FromString(BuildGroup->CommitIdentifier))
 				];
@@ -1625,6 +1746,7 @@ SBuildGroupTableRow::GenerateWidgetForColumn(const FName& ColumnName)
 	{
 		return SNew(SHorizontalBox)
 				+ SHorizontalBox::Slot()
+				.Padding(1.f,1.f)
 				[
 					SNew(STextBlock).Text(FText::FromString(BuildGroup->Suffix))
 				];
@@ -1633,6 +1755,7 @@ SBuildGroupTableRow::GenerateWidgetForColumn(const FName& ColumnName)
 	{
 		return SNew(SHorizontalBox)
 				+ SHorizontalBox::Slot()
+				.Padding(1.f,1.f)
 				[
 					SNew(STextBlock).Text(FText::FromString(BuildGroup->Category))
 				];
@@ -1643,6 +1766,7 @@ SBuildGroupTableRow::GenerateWidgetForColumn(const FName& ColumnName)
 		{
 			return SNew(SHorizontalBox)
 				+ SHorizontalBox::Slot()
+				.Padding(1.f,1.f)
 				[
 					SNew(STextBlock).Text(FText::AsDateTime(BuildGroup->CreatedAt, EDateTimeStyle::Short))
 				];
