Ensure that unix socket path is given to third parties is null-terminated.

#rnx

diff --git a/Engine/Source/Developer/Zen/Private/ZenServerHttp.cpp b/Engine/Source/Developer/Zen/Private/ZenServerHttp.cpp
--- a/Engine/Source/Developer/Zen/Private/ZenServerHttp.cpp
+++ b/Engine/Source/Developer/Zen/Private/ZenServerHttp.cpp
@@ -122,9 +122,9 @@ namespace UE::Zen {
 
 		if (Endpoint.GetSocketType() == FZenServiceEndpoint::ESocketType::Unix)
 		{
-			FStringView UnixSocketPath = Endpoint.GetName();
-			const auto& AsUtf8 = StringCast<UTF8CHAR>(UnixSocketPath.GetData(), UnixSocketPath.Len());
-			curl_easy_setopt(Curl, CURLOPT_UNIX_SOCKET_PATH, AsUtf8.Get());
+			TAnsiStringBuilder<256> AsUtf8;
+			AsUtf8 << Endpoint.GetName();
+			curl_easy_setopt(Curl, CURLOPT_UNIX_SOCKET_PATH, AsUtf8.ToString());
 		}
 	}
 
