Rework zen version check when launching/installing zenserver using new `--output-path` argument available in zenserver release 5.5.13

#rnx

diff --git a/Engine/Source/Developer/Zen/Private/ZenServerInterface.cpp b/Engine/Source/Developer/Zen/Private/ZenServerInterface.cpp
--- a/Engine/Source/Developer/Zen/Private/ZenServerInterface.cpp
+++ b/Engine/Source/Developer/Zen/Private/ZenServerInterface.cpp
@@ -24,7 +24,6 @@
 #include "Misc/ConfigCacheIni.h"
 #include "Misc/ConfigContext.h"
 #include "Misc/FileHelper.h"
-#include "Misc/MonitoredProcess.h"
 #include "Misc/Paths.h"
 #include "Misc/PathViews.h"
 #include "Misc/ScopedSlowTask.h"
@@ -279,58 +278,121 @@ GetInstallCrashpadHandlerFilePath(const FString& InTreePath)
 
 static bool GetZenVersion(const FString& UtilityPath, const FString& ServicePath, FZenVersion& OutVersion)
 {
-	FString AbsoluteUtilityPath = FPaths::ConvertRelativePathToFull(UtilityPath);
-	FMonitoredProcess MonitoredUtilityProcess(AbsoluteUtilityPath, TEXT("version --detailed"), FPaths::GetPath(UtilityPath), true);
+	FString TempOutputFilePath = FPaths::CreateTempFilename(FPlatformProcess::UserTempDir(), TEXT("zenversion"));
+	FPaths::MakePlatformFilename(TempOutputFilePath);
 
-	bool bLaunched = MonitoredUtilityProcess.Launch();
-	if (!bLaunched)
+	ON_SCOPE_EXIT
 	{
-		UE_LOG(LogZenServiceInstance, Warning, TEXT("Failed to launch zen utility to gather version data: '%s'."), *UtilityPath);
+		IFileManager& FileManager = IFileManager::Get();
+		bool DeleteResult = FileManager.Delete(*TempOutputFilePath, false, false, true);
+	};
+
+	FString AbsoluteUtilityPath = FPaths::ConvertRelativePathToFull(UtilityPath);
+	FPaths::MakePlatformFilename(AbsoluteUtilityPath);
+
+	TStringBuilder<512> Args;
+	Args << TEXTVIEW("version --detailed --output-path \"");
+	Args << TempOutputFilePath;
+	Args << TEXTVIEW("\"");
+
+	bool bLaunchDetached = true;
+	bool bLaunchHidden = true;
+	bool bLaunchReallyHidden = true;
+	uint32* OutProcessID = nullptr;
+	int32 PriorityModifier = 0;
+	void* PipeWriteChild = nullptr;
+	void* PipeReadChild = nullptr;
+
+	UE_LOG(LogZenServiceInstance, Display, TEXT("Launching zen utility '%s %s' to determine zen version."), *AbsoluteUtilityPath, *Args);
+
+	FProcHandle ZenUtilProc = FPlatformProcess::CreateProc(
+		*AbsoluteUtilityPath,
+		*Args,
+		bLaunchDetached,
+		bLaunchHidden,
+		bLaunchReallyHidden,
+		OutProcessID,
+		PriorityModifier,
+		nullptr,
+		PipeWriteChild,
+		PipeReadChild);
+
+	if (!ZenUtilProc.IsValid())
+	{
+		UE_LOG(LogZenServiceInstance, Warning, TEXT("Failed to launch zen utility to gather version data"));
 		return false;
 	}
 
+	const double MaxWaitTime = FApp::IsUnattended() ? 15.0 : 5.0;
+
 	const uint64 StartTime = FPlatformTime::Cycles64();
-	while (MonitoredUtilityProcess.Update())
+	while (FPlatformProcess::IsProcRunning(ZenUtilProc))
 	{
 		double Duration = FPlatformTime::ToSeconds64(FPlatformTime::Cycles64() - StartTime);
-		if (Duration > 10.0)
+		if (Duration >= 2.0)
 		{
-			MonitoredUtilityProcess.Cancel(true);
-			UE_LOG(LogZenServiceInstance, Warning, TEXT("Cancelled launch of zen utility for gathering version data: '%s' due to timeout."), *UtilityPath);
-
-			// Wait for execution to be terminated
-			while (MonitoredUtilityProcess.Update())
+			if (Duration >= MaxWaitTime)
 			{
-				Duration = FPlatformTime::ToSeconds64(FPlatformTime::Cycles64() - StartTime);
-				if (Duration > 15.0)
-				{
-					UE_LOG(LogZenServiceInstance, Warning, TEXT("Cancelled launch of zen utility for gathering version data: '%s'. Failed waiting for termination."), *UtilityPath);
-					break;
-				}
-				FPlatformProcess::Sleep(0.2f);
+				UE_LOG(LogZenServiceInstance, Warning, TEXT("Zen utility process has been running for %.1lf seconds without completing, aborting."), Duration);
+				FPlatformProcess::TerminateProc(ZenUtilProc);
+				FPlatformProcess::CloseProc(ZenUtilProc);
+				return false;
 			}
+			UE_LOG(LogZenServiceInstance, Warning, TEXT("Zen utility process has been running for %.1lf seconds without completing, still waiting..."), Duration);
+			FPlatformProcess::Sleep(1.0f);
+		}
+		else if (Duration >= 0.2f)
+		{
+			FPlatformProcess::Sleep(0.1f);
+		}
+		else
+		{
+			FPlatformProcess::Sleep(0.05f);
+		}
+	}
+
+	double Duration = FPlatformTime::ToSeconds64(FPlatformTime::Cycles64() - StartTime);
+
+	int ReturnCode = -1;
+	bool GetReturnCodeSucceeded = FPlatformProcess::GetProcReturnCode(ZenUtilProc, &ReturnCode);
+	FPlatformProcess::CloseProc(ZenUtilProc);
 
-			FString OutputString = MonitoredUtilityProcess.GetFullOutputWithoutDelegate();
-			UE_LOG(LogZenServiceInstance, Warning, TEXT("Launch of zen utility for gathering version data: '%s' failed. Output: '%s'"), *UtilityPath, *OutputString);
+	if (GetReturnCodeSucceeded)
+	{
+		if (ReturnCode != 0)
+		{
+			UE_LOG(LogZenServiceInstance, Warning, TEXT("Zen utility process completed after %.3lf seconds but failed with error code %d"), Duration, ReturnCode);
 			return false;
 		}
-		FPlatformProcess::Sleep(0.1f);
+	}
+	else
+	{
+		UE_LOG(LogZenServiceInstance, Warning, TEXT("Zen utility process completed after %.3lf seconds but failed to fetch return code of process"), Duration);
+		return false;
 	}
 
-	FString OutputString = MonitoredUtilityProcess.GetFullOutputWithoutDelegate();
-	if (MonitoredUtilityProcess.GetReturnCode() != 0)
+	if (!FPaths::FileExists(TempOutputFilePath))
 	{
-		UE_LOG(LogZenServiceInstance, Warning, TEXT("Unexpected return code after launch of zen utility for gathering version data: '%s' (%d). Output: '%s'"), *UtilityPath, MonitoredUtilityProcess.GetReturnCode(), *OutputString);
+		UE_LOG(LogZenServiceInstance, Warning, TEXT("Zen utility process completed successfully after %.3lf seconds but failed to produce an output file to '%s'"), Duration, *TempOutputFilePath);
 		return false;
 	}
 
-	FString VersionOutputString = OutputString.TrimStartAndEnd();
+	FString VersionString;
+	if (!FFileHelper::LoadFileToString(VersionString, *TempOutputFilePath))
+	{
+		UE_LOG(LogZenServiceInstance, Warning, TEXT("Zen utility process completed successfully after %.3lf seconds but failed to read output file at '%s'"), Duration, *TempOutputFilePath);
+		return false;
+	}
+
+	FString VersionOutputString = VersionString.TrimStartAndEnd();
 
 	if (!OutVersion.TryParse(*VersionOutputString))
 	{
-		UE_LOG(LogZenServiceInstance, Warning, TEXT("Invalid version information after launch of zen utility for gathering version data: '%s' (`%s`)"), *UtilityPath, *VersionOutputString);
+		UE_LOG(LogZenServiceInstance, Warning, TEXT("Zen utility process completed successfully after %.3lf seconds but output file is malformed : '%s'"), Duration , *VersionString);
 		return false;
 	}
+
+	UE_LOG(LogZenServiceInstance, Display, TEXT("Zen utility process completed successfully after %.3lf seconds, version: '%s'"), Duration, *VersionString);
 	return true;
 }
 
@@ -353,19 +415,24 @@ GetZenVersion(const FString& UtilityPath, const FString& ServicePath, const FStr
 			FZenVersion CachedVersion;
 			if (CachedVersion.TryParse(*VersionFileContents))
 			{
+				UE_LOG(LogZenServiceInstance, Display, TEXT("Read zen version cache file from '%s', version: '%s'"), *VersionCachePath, *VersionFileContents);
 				return CachedVersion;
 			}
 		}
 	}
-	auto GetFallbackVersion = [UtilityExecutableModificationTime, ServiceExecutableModificationTime]()
+	auto GetFallbackVersion = [UtilityExecutableModificationTime, ServiceExecutableModificationTime, &UtilityPath]()
 		{
 			FZenVersion FallbackVersion;
 			if (UtilityExecutableModificationTime > ServiceExecutableModificationTime)
 			{
 				FallbackVersion.Details = UtilityExecutableModificationTime.ToString();
-				return FallbackVersion;
 			}
-			FallbackVersion.Details = ServiceExecutableModificationTime.ToString();
+			else
+			{
+				FallbackVersion.Details = ServiceExecutableModificationTime.ToString();
+			}
+			FString VersionString = FallbackVersion.ToString();
+			UE_LOG(LogZenServiceInstance, Display, TEXT("Using legacy zen version for '%s', version: '%s'"), *UtilityPath, *VersionString);
 			return FallbackVersion;
 		};
 
@@ -373,12 +440,8 @@ GetZenVersion(const FString& UtilityPath, const FString& ServicePath, const FStr
 	FZenVersion Version;
 	if (!GetZenVersion(UtilityPath, ServicePath, Version))
 	{
-		UE_LOG(LogZenServiceInstance, Warning, TEXT("Unable to determine version using zen utility executable path: '%s'. Retrying..."), *UtilityPath);
-		if (!GetZenVersion(UtilityPath, ServicePath, Version))
-		{
-			checkf(false, TEXT("Unable to determine version using zen utility executable path: '%s'."), *UtilityPath);
-			Version = GetFallbackVersion();
-		}
+		UE_LOG(LogZenServiceInstance, Warning, TEXT("Unable to determine version using zen utility executable path: '%s'. Using legacy version check."), *UtilityPath);
+		Version = GetFallbackVersion();
 	}
 	FFileHelper::SaveStringToFile(Version.ToString(), *VersionCachePath);
 	return Version;
@@ -403,7 +466,7 @@ PromptUserToSyncInTreeVersion(const FString& ServerFilePath)
 }
 
 static bool
-IsInstallVersionOutOfDate(const FString& InTreeUtilityPath, const FString& InstallUtilityPath, const FString& InTreeServicePath, const FString& InstallServicePath, const FString& InTreeVersionCache, const FString& InstallVersionCache)
+IsInstallVersionOutOfDate(const FString& InTreeUtilityPath, const FString& InstallUtilityPath, const FString& InTreeServicePath, const FString& InstallServicePath, const FString& InTreeVersionCache, const FString& InstallVersionCache, FZenVersion& OutWantedVersion)
 {
 	IFileManager& FileManager = IFileManager::Get();
 	if (!FileManager.FileExists(*InTreeUtilityPath) || !FileManager.FileExists(*InTreeServicePath))
@@ -420,6 +483,7 @@ IsInstallVersionOutOfDate(const FString& InTreeUtilityPath, const FString& Insta
 	if (!FileManager.FileExists(*InstallUtilityPath) || !FileManager.FileExists(*InstallServicePath))
 	{
 		UE_LOG(LogZenServiceInstance, Log, TEXT("No installation found at '%s'"), *InstallServicePath);
+		OutWantedVersion = InTreeVersion;
 		return true;
 	}
 	FZenVersion InstallVersion = GetZenVersion(InstallUtilityPath, InstallServicePath, InstallVersionCache);
@@ -428,11 +492,13 @@ IsInstallVersionOutOfDate(const FString& InTreeUtilityPath, const FString& Insta
 	if (InstallVersion < InTreeVersion)
 	{
 		UE_LOG(LogZenServiceInstance, Log, TEXT("Installed version at '%s' (%s) is older than '%s' (%s)"), *InstallServicePath, *InstallVersion.ToString(), *InTreeServicePath, *InTreeVersion.ToString());
+		OutWantedVersion = InTreeVersion;
 		return true;
 	}
 	if (FParse::Param(FCommandLine::Get(), TEXT("ForceZenInstall")))
 	{
 		UE_LOG(LogZenServiceInstance, Display, TEXT("Forcing install from '%s' (%s) over '%s' (%s)"), *InTreeServicePath, *InTreeVersion.ToString(), *InstallServicePath, *InstallVersion.ToString());
+		OutWantedVersion = InTreeVersion;
 		return true;
 	}
 	return false;
@@ -2352,6 +2418,8 @@ FZenServiceInstance::ConditionalUpdateLocalInstall(FServiceAutoLaunchSettings::E
 				UE_LOG(LogZenServiceInstance, Log, TEXT("Failed cleaning up file {%s} (not critical)"), *FileToCleanUp);
 			}
 		}
+
+		UE_LOG(LogZenServiceInstance, Display, TEXT("Installed service link '%s' to '%s', version %s"), *LinkPath, *InTreeServicePath, *InTreeVersion.ToString(false));
 		return true;
 	}
 	else if (InstallMode == FServiceAutoLaunchSettings::EInstallMode::Copy)
@@ -2364,12 +2432,14 @@ FZenServiceInstance::ConditionalUpdateLocalInstall(FServiceAutoLaunchSettings::E
 		FString InstallServicePath = GetServiceCopyInstallPath();
 		FString InstallVersionCache = GetInstallVersionCachePath();
 
+		FZenVersion VersionToInstall;
+
 		IFileManager& FileManager = IFileManager::Get();
 
 		bool bMainExecutablesUpdated = false;
-		if (IsInstallVersionOutOfDate(InTreeUtilityPath, InstallUtilityPath, InTreeServicePath, InstallServicePath, InTreeVersionCache, InstallVersionCache))
+		if (IsInstallVersionOutOfDate(InTreeUtilityPath, InstallUtilityPath, InTreeServicePath, InstallServicePath, InTreeVersionCache, InstallVersionCache, VersionToInstall))
 		{
-			UE_LOG(LogZenServiceInstance, Display, TEXT("Installing service from '%s' to '%s'"), *InTreeServicePath, *InstallServicePath);
+			UE_LOG(LogZenServiceInstance, Display, TEXT("Installing service version '%s', from '%s' to '%s'"), *VersionToInstall.ToString(), *InTreeServicePath, *InstallServicePath);
 			if (!ShutDownZenServerProcessExecutable(InstallServicePath))
 			{
 				PromptUserToStopRunningServerInstanceForUpdate(InstallServicePath);
@@ -2394,6 +2464,12 @@ FZenServiceInstance::ConditionalUpdateLocalInstall(FServiceAutoLaunchSettings::E
 			WriteLocalPluginsConfig();
 
 			bMainExecutablesUpdated = true;
+
+			UE_LOG(LogZenServiceInstance, Display, TEXT("Installed service version '%s' from '%s' to '%s'"), *VersionToInstall.ToString(), *InTreeServicePath, *InstallServicePath);
+		}
+		else
+		{
+			UE_LOG(LogZenServiceInstance, Display, TEXT("Installed service at '%s' is up to date"), *InstallServicePath);
 		}
 
 #if PLATFORM_WINDOWS
