If zen.exe command tool fails to complete in under 2 seconds, write a minidump to help diagnose it

#rnx

diff --git a/Engine/Source/Developer/Zen/Private/ZenServerInterface.cpp b/Engine/Source/Developer/Zen/Private/ZenServerInterface.cpp
--- a/Engine/Source/Developer/Zen/Private/ZenServerInterface.cpp
+++ b/Engine/Source/Developer/Zen/Private/ZenServerInterface.cpp
@@ -41,6 +41,7 @@
 #	include "Windows/AllowWindowsPlatformTypes.h"
 #	include <Windows.h>
 #	include "Windows/HideWindowsPlatformTypes.h"
+#	include <dbghelp.h>
 #endif
 
 #define ALLOW_SETTINGS_OVERRIDE_FROM_COMMANDLINE			(UE_SERVER || !(UE_BUILD_SHIPPING))
@@ -275,6 +276,38 @@ GetInstallCrashpadHandlerFilePath(const FString& InTreePath)
 	return FPaths::ConvertRelativePathToFull(FPaths::Combine(GetLocalInstallPath(), FString(FPathViews::GetCleanFilename(InTreePath))));
 }
 
+#if PLATFORM_WINDOWS
+
+FString WriteMiniDump(HANDLE ProcessHandle, const FString& ExeName, const FString& OutputFolderPath)
+{
+	time_t rawtime;
+	struct tm timeinfo;
+
+	time(&rawtime);
+	localtime_s(&timeinfo, &rawtime);
+
+	FString FileName;
+	FileName.Appendf(TEXT("%s_Dump_%04d%02d%02d_%02d%02d%02d.dmp"),
+		*ExeName, timeinfo.tm_year + 1900, timeinfo.tm_mon, timeinfo.tm_mday, timeinfo.tm_hour, timeinfo.tm_min, timeinfo.tm_sec);
+
+	FString DumpPath = FPaths::Combine(OutputFolderPath, FileName);
+
+	HANDLE FileHandle = ::CreateFile(*DumpPath, GENERIC_WRITE, FILE_SHARE_WRITE, NULL, CREATE_ALWAYS, FILE_ATTRIBUTE_NORMAL, NULL);
+	if (FileHandle != INVALID_HANDLE_VALUE)
+	{
+		MINIDUMP_TYPE MinidumpType = (MINIDUMP_TYPE)(MiniDumpWithFullMemory | MiniDumpWithFullMemoryInfo | MiniDumpWithHandleData | MiniDumpWithThreadInfo);
+
+		const BOOL Result = MiniDumpWriteDump(ProcessHandle, GetProcessId(ProcessHandle), FileHandle, MinidumpType, NULL, NULL, NULL);
+		CloseHandle(FileHandle);
+		if (Result)
+		{
+			return DumpPath;
+		}
+	}
+	return FString{};
+}
+
+#endif // PLATFORM_WINDOWS
 
 static bool GetZenVersion(const FString& UtilityPath, const FString& ServicePath, FZenVersion& OutVersion)
 {
@@ -323,7 +356,9 @@ static bool GetZenVersion(const FString& UtilityPath, const FString& ServicePath
 		return false;
 	}
 
-	const double MaxWaitTime = FApp::IsUnattended() ? 15.0 : 5.0;
+	bool DumpWritten = false;
+
+	const double MaxWaitDuration = FApp::IsUnattended() ? 15.0 : 5.0;
 
 	const uint64 StartTime = FPlatformTime::Cycles64();
 	while (FPlatformProcess::IsProcRunning(ZenUtilProc))
@@ -331,14 +366,35 @@ static bool GetZenVersion(const FString& UtilityPath, const FString& ServicePath
 		double Duration = FPlatformTime::ToSeconds64(FPlatformTime::Cycles64() - StartTime);
 		if (Duration >= 2.0)
 		{
-			if (Duration >= MaxWaitTime)
+			if (Duration > 5.0)
+			{
+#if PLATFORM_WINDOWS
+				if (!DumpWritten)
+				{
+					const FString ZenExeName = FPaths::GetBaseFilename(UtilityPath);
+					const FString AutomationToolSavedDir = FPaths::Combine(FPaths::ConvertRelativePathToFull(FPaths::EngineDir()), TEXT("Programs"), TEXT("AutomationTool"), TEXT("Saved"), TEXT("Logs"));
+					const FString DumpPath = WriteMiniDump(ZenUtilProc.Get(), ZenExeName, AutomationToolSavedDir);
+					if (DumpPath.IsEmpty())
+					{
+						UE_LOG(LogZenServiceInstance, Warning, TEXT("Zen utility process has been running for %.1lf seconds without completing, failed to write minidump for zen utility to %s"), Duration, *AutomationToolSavedDir);
+					}
+					else
+					{
+						UE_LOG(LogZenServiceInstance, Warning, TEXT("Zen utility process has been running for %.1lf seconds without completing, wrote minidump for zen utility to %s"), Duration, *DumpPath);
+					}
+					DumpWritten = true;
+				}
+#endif // PLATFORM_WINDOWS
+			}
+
+			if (Duration >= MaxWaitDuration)
 			{
 				UE_LOG(LogZenServiceInstance, Warning, TEXT("Zen utility process has been running for %.1lf seconds without completing, aborting."), Duration);
 				FPlatformProcess::TerminateProc(ZenUtilProc);
 				FPlatformProcess::CloseProc(ZenUtilProc);
 				return false;
 			}
-			UE_LOG(LogZenServiceInstance, Warning, TEXT("Zen utility process has been running for %.1lf seconds without completing, still waiting..."), Duration);
+			UE_LOG(LogZenServiceInstance, Display, TEXT("Zen utility process has been running for %.1lf seconds without completing, still waiting..."), Duration);
 			FPlatformProcess::Sleep(1.0f);
 		}
 		else if (Duration >= 0.2f)
@@ -2570,17 +2626,21 @@ FZenServiceInstance::AutoLaunch(const FServiceAutoLaunchSettings& InSettings, co
 	FString WorkingDirectory = FPaths::GetPath(ExecutablePath);
 
 	ZenLockFileData LockFileState;
-	uint64 ZenWaitForRunningProcessReadyStartTime = FPlatformTime::Cycles64();
-	while (IsZenProcessUsingDataDir(*LockFilePath, &LockFileState) && LockFileState.IsValid && !LockFileState.IsReady)
 	{
-		// Server is starting up, wait for it to get ready
-		double ZenWaitDuration = FPlatformTime::ToSeconds64(FPlatformTime::Cycles64() - ZenWaitForRunningProcessReadyStartTime);
-		if (ZenWaitDuration > 5.0)
+		const double MaxWaitDuration = 30;	// Same wait time as recovery timeout
+
+		uint64 ZenWaitForRunningProcessReadyStartTime = FPlatformTime::Cycles64();
+		while (IsZenProcessUsingDataDir(*LockFilePath, &LockFileState) && LockFileState.IsValid && !LockFileState.IsReady)
 		{
-			break;
+			// Server is starting up, wait for it to get ready
+			double ZenWaitDuration = FPlatformTime::ToSeconds64(FPlatformTime::Cycles64() - ZenWaitForRunningProcessReadyStartTime);
+			if (ZenWaitDuration > MaxWaitDuration)
+			{
+				break;
+			}
+			FPlatformProcess::Sleep(0.1f);
+			LockFileState = ZenLockFileData();
 		}
-		FPlatformProcess::Sleep(0.1f);
-		LockFileState = ZenLockFileData();
 	}
 
 	bool bShutDownExistingInstanceForDataPath = true;
