Fix a typo in the spelling of zenserver.

diff --git a/Engine/Source/Developer/Zen/Private/ZenServerInterface.cpp b/Engine/Source/Developer/Zen/Private/ZenServerInterface.cpp
--- a/Engine/Source/Developer/Zen/Private/ZenServerInterface.cpp
+++ b/Engine/Source/Developer/Zen/Private/ZenServerInterface.cpp
@@ -1483,7 +1483,7 @@ static bool ShutDownZenServerProcessLockingDataDir(const FString& DataPath, doub
 				}
 				else
 				{
-					UE_LOG(LogZenServiceInstance, Warning, TEXT("Timed out waiting for shut down of zensever process using lock file '%s' with effective port %u"), *LockFilePath, EffectivePort);
+					UE_LOG(LogZenServiceInstance, Warning, TEXT("Timed out waiting for shut down of zenserver process using lock file '%s' with effective port %u"), *LockFilePath, EffectivePort);
 					break;
 				}
 			}
