Increasing the time that zen will wait to get a named mutex for startup from 5 seconds to 10 seconds.

diff --git a/Engine/Source/Developer/Zen/Private/ZenServerInterface.cpp b/Engine/Source/Developer/Zen/Private/ZenServerInterface.cpp
--- a/Engine/Source/Developer/Zen/Private/ZenServerInterface.cpp
+++ b/Engine/Source/Developer/Zen/Private/ZenServerInterface.cpp
@@ -2242,7 +2242,7 @@ FZenServiceInstance::Initialize()
 			while (true)
 			{
 				{
-					FSystemWideCriticalSection RecoveryCriticalSection(TEXT("ZenServerLaunch"), FTimespan::FromSeconds(5));
+					FSystemWideCriticalSection RecoveryCriticalSection(TEXT("ZenServerLaunch"), FTimespan::FromSeconds(10));
 					if (!RecoveryCriticalSection.IsValid())
 					{
 						UE_LOG(LogZenServiceInstance, Warning, TEXT("Local ZenServer AutoLaunch initialization timed out waiting for other process to complete startup/recovery"));
