Ensure that when we launch zenserver as a subprocess for Mac & Linux, that we detach it by specifing the --detach commandline argument to it, which will cause it to call setsid() internally.

diff --git a/Engine/Source/Developer/Zen/Private/ZenServerInterface.cpp b/Engine/Source/Developer/Zen/Private/ZenServerInterface.cpp
--- a/Engine/Source/Developer/Zen/Private/ZenServerInterface.cpp
+++ b/Engine/Source/Developer/Zen/Private/ZenServerInterface.cpp
@@ -1811,6 +1811,9 @@ StartLocalService(const FZenLocalServiceRunContext& Context)
 		return false;
 	}
 	TransientParms.Appendf(TEXT(" --child-id %s"), *StartupEventName);
+#if PLATFORM_MAC || PLATFORM_UNIX
+	TransientParms.Append(TEXT(" --detach"));
+#endif
 
 	FProcHandle Proc = StartLocalService(Context, TransientParms);
 	if (Proc.IsValid())
