When launching zenserver on Mac and Linux, ensure that we redirect its write pipe (both stdout and stderr) to /dev/null to avoid a problem where a launching process such as UAT can deadlock waiting for stdout/stderr streams of an unreal process to reach EOF as that won't happen if zenserver was launched as a child of the unreal process and inherited the stdout/stderr handles from it.  This was specifically encountered in the call to Process::WaitForExit in ProcessUtils.cs.

See public discussion of this class of issue within dotnet here:
[REDACTED_PROFILE_URL]/runtime/issues/29232#issuecomment-1451584094

diff --git a/Engine/Source/Developer/Zen/Private/ZenServerInterface.cpp b/Engine/Source/Developer/Zen/Private/ZenServerInterface.cpp
--- a/Engine/Source/Developer/Zen/Private/ZenServerInterface.cpp
+++ b/Engine/Source/Developer/Zen/Private/ZenServerInterface.cpp
@@ -1783,6 +1783,27 @@ StartLocalService(const FZenLocalServiceRunContext& Context, FStringView Transie
 		int32 PriorityModifier = 0;
 		void* PipeWriteChild = nullptr;
 		void* PipeReadChild = nullptr;
+#if PLATFORM_MAC
+		int devnullfd = open("/dev/null", O_WRONLY);
+		void* NullPipe = [[NSFileHandle alloc] initWithFileDescriptor: devnullfd];
+		PipeWriteChild = NullPipe;
+		ON_SCOPE_EXIT
+		{
+			close([(NSFileHandle*)NullPipe fileDescriptor]);
+			[(NSFileHandle*)NullPipe release];
+			close(devnullfd);
+		};
+#elif PLATFORM_UNIX
+		int devnullfd = open("/dev/null", O_WRONLY);
+		void* NullPipe = new FPipeHandle(devnullfd, devnullfd);
+		PipeWriteChild = NullPipe;
+		ON_SCOPE_EXIT
+		{
+			FPipeHandle * PipeHandle = reinterpret_cast< FPipeHandle* >(NullPipe);
+			delete PipeHandle;
+			close(devnullfd);
+		};
+#endif
 		Proc = FPlatformProcess::CreateProc(
 			*Context.GetExecutable(),
 			*Parms,
