ZenServerInterface: Add timeouts to all PerformBlockingDownload calls, short ones for optional information, long ones for messages that are fatal to the cook if not received. An infinite wait on a taskthread will deadlock the end of the cook, because we call FScheduler::RestartWorkers before finishing, and RestartWorkers requires that all task threads complete before it returns.
#rnx

diff --git a/Engine/Source/Developer/Zen/Private/ZenServerInterface.cpp b/Engine/Source/Developer/Zen/Private/ZenServerInterface.cpp
--- a/Engine/Source/Developer/Zen/Private/ZenServerInterface.cpp
+++ b/Engine/Source/Developer/Zen/Private/ZenServerInterface.cpp
@@ -2888,7 +2888,7 @@ FZenServiceInstance::GetCacheStats(FZenCacheStats& Stats)
 		// We've not got any requests in flight and we've met a given time requirement for requests
 	CacheStatsRequest = Async(ThreadPool, [&Endpoint = Endpoint]
 		{
-			UE::Zen::FZenHttpRequest Request(Endpoint, false);
+			UE::Zen::FZenHttpRequest Request(Endpoint, false, 500);
 
 			TArray64<uint8> GetBuffer;
 			FZenHttpRequest::Result Result = Request.PerformBlockingDownload(TEXTVIEW("/stats/z$"), &GetBuffer, Zen::EContentType::CbObject);
@@ -2942,7 +2942,7 @@ FZenServiceInstance::GetProjectStats(FZenProjectStats& Stats)
 			// We've not got any requests in flight and we've met a given time requirement for requests
 	ProjectStatsRequest = Async(ThreadPool, [&Endpoint = Endpoint]
 		{
-			UE::Zen::FZenHttpRequest Request(Endpoint, false);
+			UE::Zen::FZenHttpRequest Request(Endpoint, false, 500);
 
 			TArray64<uint8> GetBuffer;
 			FZenHttpRequest::Result Result = Request.PerformBlockingDownload(TEXTVIEW("/stats/prj"), &GetBuffer, Zen::EContentType::CbObject);
@@ -2995,7 +2995,7 @@ FZenServiceInstance::GetGCStatus(FGCStatus& Status)
 		// We've not got any requests in flight and we've met a given time requirement for requests
 		GCStatusRequest = Async(ThreadPool, [this]
 			{
-				UE::Zen::FZenHttpRequest Request(Endpoint, false);
+				UE::Zen::FZenHttpRequest Request(Endpoint, false, 500);
 
 				TArray64<uint8> GetBuffer;
 				FZenHttpRequest::Result Result = Request.PerformBlockingDownload(TEXTVIEW("/admin/gc"), &GetBuffer, Zen::EContentType::CbObject);
@@ -3315,7 +3315,7 @@ FZenServiceInstance::GetWorkspaces(FZenWorkspaces& Workspaces) const
 
 	WorkspacesRequest = Async(ThreadPool, [&Endpoint = Endpoint]
 	{
-		UE::Zen::FZenHttpRequest Request(Endpoint, false);
+		UE::Zen::FZenHttpRequest Request(Endpoint, false, 500);
 
 		FZenHttpRequest::Result Result = Request.PerformBlockingDownload(TEXTVIEW("/ws"), nullptr, Zen::EContentType::CbObject);
 
