Allow running Zen as a system service on Windows.

[RN] none

diff --git a/Engine/Source/Developer/Zen/Private/ZenServerInterface.cpp b/Engine/Source/Developer/Zen/Private/ZenServerInterface.cpp
--- a/Engine/Source/Developer/Zen/Private/ZenServerInterface.cpp
+++ b/Engine/Source/Developer/Zen/Private/ZenServerInterface.cpp
@@ -181,6 +181,33 @@ GetUtilityExecutableName()
 		;
 }
 
+static FString
+GetSystemInstallPath(const FSystemServiceSettings& InSettings)
+{
+	FString InstallPath = InSettings.InstallPath.IsEmpty() ? FPaths::Combine(FPlatformProcess::ApplicationSettingsDir(), TEXT("Zen/")) : InSettings.InstallPath;
+	return FPaths::ConvertRelativePathToFull(InstallPath);
+}
+
+static FString
+GetSystemServiceInstallPath(const FSystemServiceSettings& InSettings)
+{
+	return FPaths::ConvertRelativePathToFull(GetSystemInstallPath(InSettings), GetServiceExecutableName());
+}
+
+static FString
+GetSystemUtilityInstallPath(const FSystemServiceSettings& InSettings)
+{
+	return FPaths::ConvertRelativePathToFull(GetSystemInstallPath(InSettings), GetUtilityExecutableName());
+}
+
+static FString
+GetSystemServiceVersionCachePath(const FSystemServiceSettings& InSettings)
+{
+	FString InstallUtilityPath = GetSystemUtilityInstallPath(InSettings);
+	FString InstallVersionCache = FPaths::SetExtension(InstallUtilityPath, TEXT("version"));
+	return InstallVersionCache;
+}
+
 static FString
 GetLocalInstallPath()
 {
@@ -309,25 +336,11 @@ FString WriteMiniDump(HANDLE ProcessHandle, const FString& ExeName, const FStrin
 
 #endif // PLATFORM_WINDOWS
 
-static bool GetZenVersion(const FString& UtilityPath, const FString& ServicePath, FZenVersion& OutVersion)
+static bool RunZenUtility(const FString& UtilityPath, const FString& CommandLine, FString& Output, int& ReturnCode)
 {
-	FString TempOutputFilePath = FPaths::CreateTempFilename(FPlatformProcess::UserTempDir(), TEXT("zenversion"));
-	FPaths::MakePlatformFilename(TempOutputFilePath);
-
-	ON_SCOPE_EXIT
-	{
-		IFileManager& FileManager = IFileManager::Get();
-		bool DeleteResult = FileManager.Delete(*TempOutputFilePath, false, false, true);
-	};
-
 	FString AbsoluteUtilityPath = FPaths::ConvertRelativePathToFull(UtilityPath);
 	FPaths::MakePlatformFilename(AbsoluteUtilityPath);
 
-	TStringBuilder<512> Args;
-	Args << TEXTVIEW("version --detailed --output-path \"");
-	Args << TempOutputFilePath;
-	Args << TEXTVIEW("\"");
-
 	bool bLaunchDetached = true;
 	bool bLaunchHidden = true;
 	bool bLaunchReallyHidden = true;
@@ -336,11 +349,16 @@ static bool GetZenVersion(const FString& UtilityPath, const FString& ServicePath
 	void* PipeWriteChild = nullptr;
 	void* PipeReadChild = nullptr;
 
-	UE_LOG(LogZenServiceInstance, Display, TEXT("Launching zen utility '%s %s' to determine zen version."), *AbsoluteUtilityPath, *Args);
+	UE_LOG(LogZenServiceInstance, Display, TEXT("Launching zen utility '%s %s'."), *AbsoluteUtilityPath, *CommandLine);
+
+	if (!FPlatformProcess::CreatePipe(PipeReadChild, PipeWriteChild))
+	{
+		UE_LOG(LogZenServiceInstance, Warning, TEXT("Failed to create pipes for Zen child process - output may be incomplete"));
+	}
 
 	FProcHandle ZenUtilProc = FPlatformProcess::CreateProc(
 		*AbsoluteUtilityPath,
-		*Args,
+		*CommandLine,
 		bLaunchDetached,
 		bLaunchHidden,
 		bLaunchReallyHidden,
@@ -352,7 +370,7 @@ static bool GetZenVersion(const FString& UtilityPath, const FString& ServicePath
 
 	if (!ZenUtilProc.IsValid())
 	{
-		UE_LOG(LogZenServiceInstance, Warning, TEXT("Failed to launch zen utility to gather version data"));
+		UE_LOG(LogZenServiceInstance, Warning, TEXT("Failed to launch zen utility '%s %s'"), *AbsoluteUtilityPath, *CommandLine);
 		return false;
 	}
 
@@ -409,34 +427,61 @@ static bool GetZenVersion(const FString& UtilityPath, const FString& ServicePath
 
 	double Duration = FPlatformTime::ToSeconds64(FPlatformTime::Cycles64() - StartTime);
 
-	int ReturnCode = -1;
 	bool GetReturnCodeSucceeded = FPlatformProcess::GetProcReturnCode(ZenUtilProc, &ReturnCode);
+	Output = FPlatformProcess::ReadPipe(PipeReadChild);
+
+	if (Output.IsEmpty())
+	{
+		UE_LOG(LogZenServiceInstance, Warning, TEXT("Zen utility process completed after %.3lf seconds but no output was available"), Duration);
+	}
+
 	FPlatformProcess::CloseProc(ZenUtilProc);
 
 	if (GetReturnCodeSucceeded)
 	{
-		if (ReturnCode != 0)
-		{
-			UE_LOG(LogZenServiceInstance, Warning, TEXT("Zen utility process completed after %.3lf seconds but failed with error code %d"), Duration, ReturnCode);
-			return false;
-		}
+		return true;
 	}
 	else
 	{
 		UE_LOG(LogZenServiceInstance, Warning, TEXT("Zen utility process completed after %.3lf seconds but failed to fetch return code of process"), Duration);
 		return false;
 	}
+}
+
+static bool GetZenVersion(const FString& UtilityPath, const FString& ServicePath, FZenVersion& OutVersion)
+{
+	FString TempOutputFilePath = FPaths::CreateTempFilename(FPlatformProcess::UserTempDir(), TEXT("zenversion"));
+	FPaths::MakePlatformFilename(TempOutputFilePath);
+
+	ON_SCOPE_EXIT
+	{
+		IFileManager& FileManager = IFileManager::Get();
+		bool DeleteResult = FileManager.Delete(*TempOutputFilePath, false, false, true);
+	};
+
+	TStringBuilder<512> Args;
+	Args << TEXTVIEW("version --detailed --output-path \"");
+	Args << TempOutputFilePath;
+	Args << TEXTVIEW("\"");
+
+	FString Output;
+	int ReturnCode;
+	if (!RunZenUtility(UtilityPath, *Args, Output, ReturnCode) || ReturnCode != 0)
+	{
+		UE_LOG(LogZenServiceInstance, Warning, TEXT("Zen utility process completed but failed with error code %d"), ReturnCode);
+		return false;
+	}
 
 	if (!FPaths::FileExists(TempOutputFilePath))
 	{
-		UE_LOG(LogZenServiceInstance, Warning, TEXT("Zen utility process completed successfully after %.3lf seconds but failed to produce an output file to '%s'"), Duration, *TempOutputFilePath);
+		UE_LOG(LogZenServiceInstance, Warning, TEXT("Zen utility process completed successfully but failed to produce an output file to '%s'"), *TempOutputFilePath);
 		return false;
 	}
 
 	FString VersionString;
 	if (!FFileHelper::LoadFileToString(VersionString, *TempOutputFilePath))
 	{
-		UE_LOG(LogZenServiceInstance, Warning, TEXT("Zen utility process completed successfully after %.3lf seconds but failed to read output file at '%s'"), Duration, *TempOutputFilePath);
+		UE_LOG(LogZenServiceInstance, Warning, TEXT("Zen utility process completed successfully but failed to read output file at '%s'"), *TempOutputFilePath);
 		return false;
 	}
 
@@ -444,11 +489,11 @@ static bool GetZenVersion(const FString& UtilityPath, const FString& ServicePath
 
 	if (!OutVersion.TryParse(*VersionOutputString))
 	{
-		UE_LOG(LogZenServiceInstance, Warning, TEXT("Zen utility process completed successfully after %.3lf seconds but output file is malformed : '%s'"), Duration , *VersionString);
+		UE_LOG(LogZenServiceInstance, Warning, TEXT("Zen utility process completed successfully but output file is malformed : '%s'"), *VersionString);
 		return false;
 	}
 
-	UE_LOG(LogZenServiceInstance, Display, TEXT("Zen utility process completed successfully after %.3lf seconds, version: '%s'"), Duration, *VersionString);
+	UE_LOG(LogZenServiceInstance, Display, TEXT("Zen utility process completed successfully, version: '%s'"), *VersionString);
 	return true;
 }
 
@@ -868,6 +913,25 @@ PromptUserUnableToDetermineValidDataPath()
 	}
 }
 
+static void
+PromptUserInvalidSystemServiceDataPath()
+{
+#if !IS_PROGRAM
+	if (!FApp::IsUnattended() && !IsRunningCommandlet() && !GIsRunningUnattendedScript)
+	{
+		FString LogDirPath = FPaths::ConvertRelativePathToFull(FPaths::ProjectLogDir());
+		FText ZenInvalidDataPathPromptTitle = NSLOCTEXT("Zen", "Zen_InvalidDataPathPromptTitle", "No Valid Data Path Configuration");
+		FText ZenInvalidDataPathPromptText = FText::Format(NSLOCTEXT("Zen", "Zen_InvalidSystemServiceDataPathPromptText", "Unreal Zen Storage Server can not determine a valid data path.\nPlease check the log in '{0}' for details.\nFalling back to Zen AutoLaunch or ConnectExisting settings."), FText::FromString(LogDirPath));
+		FPlatformMisc::MessageBoxExt(EAppMsgType::Ok, *ZenInvalidDataPathPromptText.ToString(), *ZenInvalidDataPathPromptTitle.ToString());
+	}
+	else
+#endif
+	{
+		// Just log as there is no one to show a message
+		UE_LOG(LogZenServiceInstance, Warning, TEXT("Unreal Zen Storage Server is unable to determine a valid data path for the system service."));
+	}
+}
+
 static void
 PromptUserAboutInvalidValidDataPathConfiguration(const FString& UsedDataPath)
 {
@@ -1063,6 +1127,32 @@ void FServicePluginSettings::WriteToCompactBinary(FCbWriter& Writer) const
 	Writer.EndObject();
 }
 
+static TOptional<FSystemServiceSettings> ReadSystemServiceSettingsFromConfig()
+{
+	const TCHAR* SystemServiceConfigSection = TEXT("Zen.SystemService");
+	FSystemServiceSettings SystemServiceSettings;
+	bool bHasInvalidPathConfigurations = false;
+	bool bIsDefaultDataPath = false;
+
+	if (!DetermineDataPath(SystemServiceConfigSection, SystemServiceSettings.DataPath, bHasInvalidPathConfigurations, bIsDefaultDataPath))
+	{
+		PromptUserInvalidSystemServiceDataPath();
+		return {};
+	}
+	else if (bHasInvalidPathConfigurations)
+	{
+		PromptUserAboutInvalidValidDataPathConfiguration(SystemServiceSettings.DataPath);
+	}
+
+	GConfig->GetString(SystemServiceConfigSection, TEXT("InstallPath"), SystemServiceSettings.InstallPath, GEngineIni);
+	GConfig->GetString(SystemServiceConfigSection, TEXT("ExtraArgs"), SystemServiceSettings.ExtraArgs, GEngineIni);
+	GConfig->GetBool(SystemServiceConfigSection, TEXT("AllowRemoteNetworkService"), SystemServiceSettings.bAllowRemoteNetworkService, GEngineIni);
+	EnsureEditorSettingsConfigLoaded();
+	GConfig->GetBool(TEXT("/Script/UnrealEd.CrashReportsPrivacySettings"), TEXT("bSendUnattendedBugReports"), SystemServiceSettings.bSendUnattendedBugReports, GEditorSettingsIni);
+
+	return TOptional<FSystemServiceSettings>(SystemServiceSettings);
+}
+
 bool
 FServiceSettings::ReadFromConfig()
 {
@@ -1070,6 +1160,17 @@ FServiceSettings::ReadFromConfig()
 	const TCHAR* ConfigSection = TEXT("Zen");
 	bool bAutoLaunch = true;
 	GConfig->GetBool(ConfigSection, TEXT("AutoLaunch"), bAutoLaunch, GEngineIni);
+	GConfig->GetBool(ConfigSection, TEXT("SystemService"), bSystemService, GEngineIni);
+
+	if (bSystemService)
+	{
+		TOptional<FSystemServiceSettings> Settings = ReadSystemServiceSettingsFromConfig();
+		if (Settings)
+		{
+			bSystemService = true;
+			SystemServiceSettings = *Settings;
+		}
+	}
 
 	if (bAutoLaunch)
 	{
@@ -1166,6 +1267,21 @@ FServiceSettings::ReadFromCompactBinary(FCbFieldView Field)
 {
 	if (Field.IsObject())
 	{
+		if (bool bSystemServiceValue = Field["bSystemService"].AsBool())
+		{
+			bSystemService = true;
+
+			if (FCbObjectView SystemSettingsObject = Field["SystemServiceSettings"].AsObjectView())
+			{
+				SystemServiceSettings.DataPath = FString(SystemSettingsObject["DataPath"].AsString());
+				SystemServiceSettings.ExtraArgs = FString(SystemSettingsObject["ExtraArgs"].AsString());
+				SystemServiceSettings.DesiredPort = SystemSettingsObject["DesiredPort"].AsInt16();
+				SystemServiceSettings.InstallPath = FString(SystemSettingsObject["InstallPath"].AsString());
+				SystemServiceSettings.bAllowRemoteNetworkService = SystemSettingsObject["AllowRemoteNetworkService"].AsBool();
+				SystemServiceSettings.bSendUnattendedBugReports = SystemSettingsObject["SendUnattendedBugReports"].AsBool();
+			}
+		}
+
 		if (bool bAutoLaunchValue = Field["bAutoLaunch"].AsBool())
 		{
 			if (!TryApplyAutoLaunchOverride())
@@ -1256,6 +1372,20 @@ FServiceSettings::WriteToCompactBinary(FCbWriter& Writer) const
 {
 	bool bAutoLaunch = IsAutoLaunch();
 	Writer << "bAutoLaunch" << bAutoLaunch;
+	Writer << "bSystemService" << IsSystemService();
+
+	if (bSystemService)
+	{
+		Writer.BeginObject("SystemServiceSettings");
+		Writer << "DataPath" << SystemServiceSettings.DataPath;
+		Writer << "ExtraArgs" << SystemServiceSettings.ExtraArgs;
+		Writer << "DesiredPort" << SystemServiceSettings.DesiredPort;
+		Writer << "InstallPath" << SystemServiceSettings.InstallPath;
+		Writer << "SendUnattendedBugReports" << SystemServiceSettings.bSendUnattendedBugReports;
+		Writer << "AllowRemoteNetworkService" << SystemServiceSettings.bAllowRemoteNetworkService;
+		Writer.EndObject();
+	}
+
 	if (bAutoLaunch)
 	{
 		const FServiceAutoLaunchSettings& AutoLaunchSettings = SettingsVariant.Get<FServiceAutoLaunchSettings>();
@@ -1545,23 +1675,29 @@ IsZenProcessUsingDataDir(const TCHAR* LockFilePath, ZenLockFileData* OutLockFile
 }
 
 static FString
-DetermineCmdLineWithoutTransientComponents(const FServiceAutoLaunchSettings& InSettings, int16 OverrideDesiredPort)
+DetermineCmdLineWithoutTransientComponents(const FString& DataPath, const FString& ExtraArgs, const TArray<FServicePluginSettings>& Plugins, bool bSendUnattendedBugReports, bool bAllowRemoteNetworkService, int16 OverrideDesiredPort)
 {
-	FString PlatformDataPath(InSettings.DataPath);
+	FString PlatformDataPath(DataPath);
 	FPaths::MakePlatformFilename(PlatformDataPath);
 
+	int32 SpaceIndex = INDEX_NONE;
+	if (PlatformDataPath.FindChar(' ', SpaceIndex))
+	{
+		PlatformDataPath = TEXT("\"") + PlatformDataPath + TEXT("\"");
+	}
+
 	FString Parms;
-	Parms.Appendf(TEXT("--port %d --data-dir \"%s\""),
+	Parms.Appendf(TEXT("--port %d --data-dir %s"),
 		OverrideDesiredPort,
 		*PlatformDataPath);
 
-	if (!InSettings.ExtraArgs.IsEmpty())
+	if (!ExtraArgs.IsEmpty())
 	{
 		Parms.AppendChar(TEXT(' '));
-		Parms.Append(InSettings.ExtraArgs);
+		Parms.Append(ExtraArgs);
 	}
 
-	if (!InSettings.Plugins.IsEmpty())
+	if (!Plugins.IsEmpty())
 	{
 		Parms.AppendChar(TEXT(' '));
 		Parms.Appendf(TEXT("--plugins-config \"%s\""),
@@ -1588,12 +1724,12 @@ DetermineCmdLineWithoutTransientComponents(const FServiceAutoLaunchSettings& InS
 		}
 	}
 
-	if (!InSettings.bSendUnattendedBugReports)
+	if (bSendUnattendedBugReports)
 	{
 		Parms.Append(TEXT(" --no-sentry"));
 	}
 
-	if (!InSettings.bAllowRemoteNetworkService)
+	if (!bAllowRemoteNetworkService)
 	{
 		Parms.Append(TEXT(" --http-forceloopback"));
 	}
@@ -2077,7 +2213,7 @@ const FString FZenServiceInstance::GetPath() const
 bool
 FZenServiceInstance::IsServiceRunning()
 {
-	return !Settings.IsAutoLaunch() || bHasLaunchedLocal;
+	return !Settings.IsAutoLaunch() || bHasLaunchedLocal || bHasLaunchedSystemService;
 }
 
 bool
@@ -2246,9 +2382,41 @@ FZenServiceInstance::GetAutoLaunchedPort()
 	return AutoLaunchedPort;
 }
 
+static bool UninstallSystemService(FString UtilityPath);
+static void PromptUserOfFailedSystemServiceUninstall();
+
 void
 FZenServiceInstance::Initialize()
 {
+	if (Settings.IsSystemService())
+	{
+		if (ConditionalUpdateSystemServiceInstall(Settings.SystemServiceSettings))
+		{
+			return;
+		}
+
+		Settings.bSystemService = false;
+	}
+
+	check(GConfig && GConfig->IsReadyForUse());
+	const TCHAR* ConfigSection = TEXT("Zen");
+	bool bSystemService = false;
+	GConfig->GetBool(ConfigSection, TEXT("SystemService"), bSystemService, GEngineIni);
+
+	if (!bSystemService)
+	{
+		FSystemWideCriticalSection InstallCriticalSection(TEXT("ZenServerInstall"), FTimespan::FromSeconds(10));
+		if (!InstallCriticalSection.IsValid())
+		{
+			PromptUserOfFailedSystemServiceUninstall();
+		}
+
+		if (!UninstallSystemService(GetInTreeUtilityPath()))
+		{
+			PromptUserOfFailedSystemServiceUninstall();
+		}
+	}
+
 	if (Settings.IsAutoLaunch())
 	{
 		uint64 ZenAutoLaunchStartTime = FPlatformTime::Cycles64();
@@ -2318,6 +2486,24 @@ FZenServiceInstance::Initialize()
 	}
 }
 
+static void
+PromptUserNeedsElevatedPermissions()
+{
+#if !IS_PROGRAM
+	if (!FApp::IsUnattended() && !IsRunningCommandlet() && !GIsRunningUnattendedScript)
+	{
+		FText ZenPermissionsPromptTitle = NSLOCTEXT("Zen", "Zen_PermissionsPromptTitle", "Administrator permissions required");
+		FText ZenPermissionsPromptText = NSLOCTEXT("Zen", "Zen_PermissionsPromptText", "Unable to update Unreal Zen Storage Server. Updating the installed system service requires administrator privileges.");
+		FPlatformMisc::MessageBoxExt(EAppMsgType::Ok, *ZenPermissionsPromptText.ToString(), *ZenPermissionsPromptTitle.ToString());
+	}
+	else
+#endif
+	{
+		// Just log as there is no one to show a message
+		UE_LOG(LogZenServiceInstance, Warning, TEXT("Unable to update Unreal Zen Storage Server. Updating the installed system service requires administrator privileges."));
+	}
+}
+
 static void
 PromptUserToStopRunningServerInstanceForUpdate(const FString& ServerFilePath)
 {
@@ -2372,6 +2558,401 @@ PromptUserOfFailedShutDownOfExistingProcess(uint16 Port)
 	}
 }
 
+static void
+PromptUserOfFailedSystemServiceInstall()
+{
+#if !IS_PROGRAM
+	if (!FApp::IsUnattended() && !IsRunningCommandlet() && !GIsRunningUnattendedScript)
+	{
+		FText ZenInstallSystemServiceFailurePromptTitle = NSLOCTEXT("Zen", "Zen_SystemServiceInstallFailurePromptTitle", "Failed to install system service");
+		FText ZenInstallSystemServiceFailurePromptText = NSLOCTEXT("Zen", "Zen_SystemServiceInstallFailurePromptText", "Unreal Zen Storage Server failed to install as a system service. Zen will run using AutoLaunch or ConnectExisting settings.");
+		FPlatformMisc::MessageBoxExt(EAppMsgType::Ok, *ZenInstallSystemServiceFailurePromptText.ToString(), *ZenInstallSystemServiceFailurePromptTitle.ToString());
+	}
+	else
+#endif
+	{
+		// Just log as there is no one to show a message
+		UE_LOG(LogZenServiceInstance, Warning, TEXT("Unreal Zen Storage Server failed to install as a system service. Zen will run using AutoLaunch or ConnectExisting settings."));
+	}
+}
+
+static void
+PromptUserOfFailedSystemServiceUninstall()
+{
+#if !IS_PROGRAM
+	if (!FApp::IsUnattended() && !IsRunningCommandlet() && !GIsRunningUnattendedScript)
+	{
+		FText ZenUninstallSystemServiceFailurePromptTitle = NSLOCTEXT("Zen", "Zen_SystemServiceUninstallFailurePromptTitle", "Failed to uninstall system service");
+		FText ZenUninstallSystemServiceFailurePromptText = NSLOCTEXT("Zen", "Zen_SystemServiceUninstallFailurePromptText", "Unreal Zen Storage Server was unable to uninstall a running system service instance. Zen AutoLaunch or ConnectExisting settings may fail to function correctly.");
+		FPlatformMisc::MessageBoxExt(EAppMsgType::Ok, *ZenUninstallSystemServiceFailurePromptText.ToString(), *ZenUninstallSystemServiceFailurePromptTitle.ToString());
+	}
+	else
+#endif
+	{
+		// Just log as there is no one to show a message
+		UE_LOG(LogZenServiceInstance, Warning, TEXT("Unreal Zen Storage Server was unable to uninstall a running system service instance. Zen AutoLaunch settings may fail to function correctly."));
+	}
+}
+
+static void
+PromptUserOfMismatchedSystemServiceConfiguration(const FString& DesiredCommandLine, const FString& ActualCommandLine)
+{
+#if !IS_PROGRAM
+	if (!FApp::IsUnattended() && !IsRunningCommandlet() && !GIsRunningUnattendedScript)
+	{
+		FText ZenMismatchedSystemServiceConfigurationPromptTitle = NSLOCTEXT("Zen", "Zen_SystemServiceMismatchedConfigurationPromptTitle", "Mismatched configuration for Unreal Zen Storage Server");
+		FText ZenMismatchedSystemServiceConfigurationPromptText = FText::Format(NSLOCTEXT("Zen", "Zen_SystemServiceMismatchedConfigurationPromptText", "Unreal Zen Storage Server detected mismatched configuration for system service installation.\n  Your configuration requested the following command line: '{0}'.\n  The installed system service is running with the command line: '{1}'\n. To force an update to the system service configuration, use -ForceZenInstall on the command line."), FText::FromString(DesiredCommandLine), FText::FromString(ActualCommandLine));
+		FPlatformMisc::MessageBoxExt(EAppMsgType::Ok, *ZenMismatchedSystemServiceConfigurationPromptText.ToString(), *ZenMismatchedSystemServiceConfigurationPromptTitle.ToString());
+	}
+	else
+#endif
+	{
+		// Just log as there is no one to show a message
+		UE_LOG(LogZenServiceInstance, Warning, TEXT("Unreal Zen Storage Server detected mismatched configuration for system service installation. Zen will run using the existing system service configuration. To force an update to the system service configuration, use -ForceZenInstall on the command line."));
+	}
+}
+
+static void
+PromptUserUnableToUpdateSystemServiceBinaries(const FString& InstallPath)
+{
+#if !IS_PROGRAM
+	if (!FApp::IsUnattended() && !IsRunningCommandlet() && !GIsRunningUnattendedScript)
+	{
+		FText ZenUnableToUpdateBinariesPromptTitle = NSLOCTEXT("Zen", "Zen_SystemServiceMismatchedConfigurationPromptTitle", "Failed to update Unreal Zen Storage Server binaries");
+		FText ZenUnableToUpdateBinariesPromptText = FText::Format(NSLOCTEXT("Zen", "Zen_SystemServiceMismatchedConfigurationPromptText", "Unreal Zen Storage Server was unable to update the binaries at '%s'. Zen will run using AutoLaunch or ConnectExisting settings."), FText::FromString(InstallPath));
+		FPlatformMisc::MessageBoxExt(EAppMsgType::Ok, *ZenUnableToUpdateBinariesPromptTitle.ToString(), *ZenUnableToUpdateBinariesPromptText.ToString());
+	}
+	else
+#endif
+	{
+		// Just log as there is no one to show a message
+		UE_LOG(LogZenServiceInstance, Warning, TEXT("Unreal Zen Storage Server was unable to update the binaries at '%s'.Zen will run using AutoLaunch or ConnectExisting settings."), *InstallPath);
+	}
+}
+
+static bool IsDirectoryWriteable(const FString& Directory)
+{
+	IFileManager& FileManager = IFileManager::Get();
+	FString NormalizedDirectory = FPaths::ConvertRelativePathToFull(Directory);
+	FPaths::NormalizeDirectoryName(NormalizedDirectory);
+	FFileStatData StatData = FileManager.GetStatData(*NormalizedDirectory);
+
+	if (StatData.bIsValid && StatData.bIsDirectory)
+	{
+		FString TestFilePath = NormalizedDirectory / FString::Printf(TEXT(".zen-startup-test-file-%d"), FPlatformProcess::GetCurrentProcessId());
+		FArchive* TestFile = FileManager.CreateFileWriter(*TestFilePath, FILEWRITE_Silent);
+		if (!TestFile)
+		{
+			return false;
+		}
+
+		TestFile->Close();
+		delete TestFile;
+		FileManager.Delete(*TestFilePath);
+	}
+	else if (!FileManager.MakeDirectory(*NormalizedDirectory, true))
+	{
+		return false;
+	}
+
+	return true;
+}
+
+static bool UpdateServiceBinaries(const FString& InTreeUtilityPath, const FString& InstalledUtilityPath, const FString& InTreeServicePath, const FString& InstalledServicePath, const FString& InTreeVersionCache, const FString& InstalledVersionCache)
+{
+	// Even after waiting for the process to shut down we have a tolerance for failure when overwriting the target files
+	if (!AttemptFileCopyWithRetries(*InstalledServicePath, *InTreeServicePath, 5.0))
+	{
+		return false;
+	}
+
+	if (!AttemptFileCopyWithRetries(*InstalledUtilityPath, *InTreeUtilityPath, 5.0))
+	{
+		return false;
+	}
+
+	AttemptFileCopyWithRetries(*InstalledVersionCache, *InTreeVersionCache, 5.0);
+	return true;
+}
+
+static bool DetermineSystemServiceInfo(const FString UtilityPath, FString& OutExecutable, FString& OutCommandLine, uint16& OutPort)
+{
+	FString CommandLine = TEXT("service status");
+	FString Output;
+	int ReturnCode;
+
+	if (!RunZenUtility(UtilityPath, CommandLine, Output, ReturnCode) || ReturnCode != 0)
+	{
+		return false;
+	}
+
+	// Versions of zen that don't know about the service command will still return 0 for 'service status', leading to warnings
+	// later on in this code when we fail to find the info in the output.
+	if (Output.StartsWith("Unknown command"))
+	{
+		return false;
+	}
+
+	FString CommandLineOptionsSubstring = TEXT("CommandLineOptions:");
+	int32 ArgsIndex = Output.Find(CommandLineOptionsSubstring) + CommandLineOptionsSubstring.Len();
+	if (ArgsIndex == INDEX_NONE)
+	{
+		UE_LOG(LogZenServiceInstance, Warning, TEXT("Failed to determine command line options for running Unreal Zen Storage Server instance"));
+		return false;
+	}
+
+	int32 ArgsEnd = Output.Find("\n", ESearchCase::IgnoreCase, ESearchDir::FromStart, ArgsIndex);
+	if (ArgsEnd == INDEX_NONE)
+	{
+		UE_LOG(LogZenServiceInstance, Warning, TEXT("Failed to determine command line options for running Unreal Zen Storage Server instance"));
+		return false;
+	}
+
+	OutCommandLine = Output.Mid(ArgsIndex, ArgsEnd - ArgsIndex).TrimStartAndEnd();
+
+	FString PortSubstring = TEXT("--port ");
+	int32 PortIndex = OutCommandLine.Find(PortSubstring) + PortSubstring.Len();
+	if (PortIndex == INDEX_NONE)
+	{
+		UE_LOG(LogZenServiceInstance, Warning, TEXT("Failed to determine port for running Unreal Zen Storage Server instance"));
+		return false;
+	}
+
+	int32 PortEnd = OutCommandLine.Find(" ", ESearchCase::IgnoreCase, ESearchDir::FromStart, PortIndex);
+	if (PortEnd == INDEX_NONE)
+	{
+		UE_LOG(LogZenServiceInstance, Warning, TEXT("Failed to determine port for running Unreal Zen Storage Server instance"));
+		return false;
+	}
+
+	if (!LexTryParseString(OutPort, *OutCommandLine.Mid(PortIndex, PortEnd - PortIndex)))
+	{
+		UE_LOG(LogZenServiceInstance, Warning, TEXT("Failed to determine port for running Unreal Zen Storage Server instance"));
+		return false;
+	}
+
+	FString ExecutableSubstring = TEXT("Executable:");
+	int32 ExecutableIndex = Output.Find(ExecutableSubstring) + ExecutableSubstring.Len();
+	if (ExecutableIndex == INDEX_NONE)
+	{
+		UE_LOG(LogZenServiceInstance, Warning, TEXT("Failed to determine executable for running Unreal Zen Storage Server instance"));
+		return false;
+	}
+
+	int32 ExecutableEnd = Output.Find("\n", ESearchCase::IgnoreCase, ESearchDir::FromStart, ExecutableIndex);
+	if (ExecutableEnd == INDEX_NONE)
+	{
+		UE_LOG(LogZenServiceInstance, Warning, TEXT("Failed to determine executable for running Unreal Zen Storage Server instance"));
+		return false;
+	}
+
+	OutExecutable = Output.Mid(ExecutableIndex, ExecutableEnd - ExecutableIndex).TrimStartAndEnd();
+
+	return true;
+}
+
+static
+bool AddCurrentProcessAsSponsor(uint16 DesiredPort)
+{
+	ZenServerState State(/*ReadOnly*/ false);
+	ZenServerState::ZenServerEntry* RunningEntry = State.LookupByDesiredListenPort(DesiredPort);
+	if (RunningEntry)
+	{
+		RunningEntry->AddSponsorProcess(FPlatformProcess::GetCurrentProcessId());
+		return true;
+	}
+	else
+	{
+		return false;
+	}
+}
+
+static
+bool UninstallSystemService(FString UtilityPath)
+{
+	FString RunningExecutable;
+	FString RunningCommandLine;
+	uint16 RunningPort;
+
+	bool bIsServiceRunning = DetermineSystemServiceInfo(UtilityPath, RunningExecutable, RunningCommandLine, RunningPort);
+
+	FString Output;
+	int32 ReturnCode;
+
+	if (bIsServiceRunning)
+	{
+		const ZenServerState State(/*ReadOnly*/ true);
+		const ZenServerState::ZenServerEntry* Entry = State.LookupByEffectiveListenPort(RunningPort);
+
+		if (!Entry)
+		{
+			UE_LOG(LogZenServiceInstance, Warning, TEXT("Failed to look up running service info to determine active users of service. Assuming no active users."));
+		}
+		else if (Entry->ActiveSponsors.load())
+		{
+			PromptUserToStopRunningServerInstanceForUpdate(FPaths::GetPath(UtilityPath));
+			return false;
+		}
+
+		if (!RunZenUtility(UtilityPath, TEXT("service stop --allow-elevation"), Output, ReturnCode))
+		{
+			PromptUserOfFailedSystemServiceUninstall();
+			return false;
+		}
+
+		if (ReturnCode == -1)
+		{
+			PromptUserNeedsElevatedPermissions();
+			return false;
+		}
+	}
+
+	if (!RunZenUtility(UtilityPath, TEXT("service status"), Output, ReturnCode))
+	{
+		PromptUserOfFailedSystemServiceUninstall();
+		return false;
+	}
+
+	// zen returns 1 for status of a service that's not installed
+	if (ReturnCode == 1)
+	{
+		return true;
+	}
+
+	if (!RunZenUtility(UtilityPath, TEXT("service uninstall --allow-elevation"), Output, ReturnCode))
+	{
+		PromptUserOfFailedSystemServiceUninstall();
+		return false;
+	}
+
+	if (ReturnCode == -1)
+	{
+		PromptUserNeedsElevatedPermissions();
+		return false;
+	}
+
+	return true;
+}
+
+bool FZenServiceInstance::ConditionalUpdateSystemServiceInstall(const FSystemServiceSettings &InSettings)
+{
+	FString SystemServiceInstallPath = GetSystemServiceInstallPath(InSettings);
+	FString SystemUtilityInstallPath = GetSystemUtilityInstallPath(InSettings);
+	FString SystemVersionCache = GetSystemServiceVersionCachePath(InSettings);
+
+	FString InTreeUtilityPath = GetInTreeUtilityPath();
+	FString InTreeServicePath = GetInTreeServicePath();
+	FString InTreeVersionCache = GetInTreeVersionCache();
+
+	FString Output;
+	int ReturnCode;
+
+	FSystemWideCriticalSection InstallCriticalSection(TEXT("ZenServerInstall"), FTimespan::FromSeconds(10));
+	if (!InstallCriticalSection.IsValid())
+	{
+		PromptUserOfFailedSystemServiceInstall();
+		return false;
+	}
+
+	FString RunningExecutable;
+	FString RunningCommandLine;
+	uint16 RunningPort;
+
+	TArray<FServicePluginSettings> Plugins;
+	FString DesiredCommandLine = DetermineCmdLineWithoutTransientComponents(InSettings.DataPath, InSettings.ExtraArgs, Plugins, InSettings.bSendUnattendedBugReports, InSettings.bAllowRemoteNetworkService, InSettings.DesiredPort);
+
+	FZenVersion Version;
+	bool bIsServiceRunning = DetermineSystemServiceInfo(SystemUtilityInstallPath, RunningExecutable, RunningCommandLine, RunningPort);
+	bool bServiceNeedsUpdate = !bIsServiceRunning
+		|| RunningCommandLine != DesiredCommandLine
+		|| IsInstallVersionOutOfDate(InTreeUtilityPath, SystemUtilityInstallPath, InTreeServicePath, SystemServiceInstallPath, InTreeVersionCache, SystemVersionCache, Version);
+
+	if (bServiceNeedsUpdate)
+	{
+		FString InstallDirectory = GetSystemInstallPath(InSettings);
+		if (!IsDirectoryWriteable(InstallDirectory))
+		{
+			PromptUserUnableToUpdateSystemServiceBinaries(InstallDirectory);
+			return false;
+		}
+
+		if (!UninstallSystemService(SystemUtilityInstallPath))
+		{
+			return false;
+		}
+
+		const ZenServerState State(/*ReadOnly*/true);
+		const ZenServerState::ZenServerEntry* RunningEntry = State.LookupByDesiredListenPort(InSettings.DesiredPort);
+		if (RunningEntry != nullptr)
+		{
+			if (!ShutdownZenServerProcess(RunningEntry->Pid))
+			{
+				PromptUserOfFailedShutDownOfExistingProcess(InSettings.DesiredPort);
+				return false;
+			}
+		}
+
+		if (!UpdateServiceBinaries(InTreeUtilityPath, SystemUtilityInstallPath, InTreeServicePath, SystemServiceInstallPath, InTreeVersionCache, SystemVersionCache))
+		{
+			PromptUserUnableToUpdateSystemServiceBinaries(InstallDirectory);
+			return false;
+		}
+
+		FString ServiceInstallCommandLine;
+		ServiceInstallCommandLine.Appendf(TEXT("service install --allow-elevation -- %s"), *DesiredCommandLine);
+
+		if (!RunZenUtility(SystemUtilityInstallPath, ServiceInstallCommandLine, Output, ReturnCode))
+		{
+			PromptUserOfFailedSystemServiceInstall();
+			return false;
+		}
+
+		if (ReturnCode == -1)
+		{
+			PromptUserNeedsElevatedPermissions();
+			return false;
+		}
+
+		if (!RunZenUtility(SystemUtilityInstallPath, TEXT("service start --allow-elevation"), Output, ReturnCode))
+		{
+			PromptUserOfFailedSystemServiceInstall();
+			return false;
+		}
+
+		if (ReturnCode == -1)
+		{
+			PromptUserNeedsElevatedPermissions();
+			return false;
+		}
+
+		bIsServiceRunning = DetermineSystemServiceInfo(SystemUtilityInstallPath, RunningExecutable, RunningCommandLine, RunningPort);
+	}
+
+	if (!bIsServiceRunning)
+	{
+		UE_LOG(LogZenServiceInstance, Warning, TEXT("Unknown error starting Unreal Zen Storage Server"));
+		return false;
+	}
+
+	AddCurrentProcessAsSponsor(RunningPort);
+
+	Endpoint = FZenServiceEndpoint(TEXT("[::1]"), RunningPort);
+
+	FZenLocalServiceRunContext EffectiveRunContext;
+	EffectiveRunContext.Executable = RunningExecutable;
+	EffectiveRunContext.CommandlineArguments = RunningCommandLine;
+	EffectiveRunContext.WorkingDirectory = SystemServiceInstallPath;
+	EffectiveRunContext.DataPath = InSettings.DataPath;
+	EffectiveRunContext.bShowConsole = false;
+	EffectiveRunContext.bLimitProcessLifetime = false;
+
+	EffectiveRunContext.WriteToJsonFile(*GetServiceRunContextPath());
+
+	bHasLaunchedSystemService = true;
+
+	return true;
+}
+
 bool
 FZenServiceInstance::ConditionalUpdateLocalInstall(FServiceAutoLaunchSettings::EInstallMode InstallMode)
 {
@@ -2685,7 +3266,7 @@ FZenServiceInstance::AutoLaunch(const FServiceAutoLaunchSettings& InSettings, co
 			{
 				FZenLocalServiceRunContext DesiredRunContext;
 				DesiredRunContext.Executable = ExecutablePath;
-				DesiredRunContext.CommandlineArguments = DetermineCmdLineWithoutTransientComponents(InSettings, InSettings.DesiredPort);
+				DesiredRunContext.CommandlineArguments = DetermineCmdLineWithoutTransientComponents(InSettings.DataPath, InSettings.ExtraArgs, InSettings.Plugins, InSettings.bSendUnattendedBugReports, InSettings.bAllowRemoteNetworkService, InSettings.DesiredPort);
 				DesiredRunContext.WorkingDirectory = WorkingDirectory;
 				DesiredRunContext.DataPath = InSettings.DataPath;
 				DesiredRunContext.bShowConsole = InSettings.bShowConsole;
@@ -2793,7 +3374,7 @@ FZenServiceInstance::AutoLaunch(const FServiceAutoLaunchSettings& InSettings, co
 			}
 		}
 
-		FString ParmsWithoutTransients = DetermineCmdLineWithoutTransientComponents(InSettings, InSettings.DesiredPort);
+		FString ParmsWithoutTransients = DetermineCmdLineWithoutTransientComponents(InSettings.DataPath, InSettings.ExtraArgs, InSettings.Plugins, InSettings.bSendUnattendedBugReports, InSettings.bAllowRemoteNetworkService, InSettings.DesiredPort);
 
 		FZenLocalServiceRunContext EffectiveRunContext;
 		EffectiveRunContext.Executable = ExecutablePath;
diff --git a/Engine/Source/Developer/Zen/Private/ZenServerState.cpp b/Engine/Source/Developer/Zen/Private/ZenServerState.cpp
--- a/Engine/Source/Developer/Zen/Private/ZenServerState.cpp
+++ b/Engine/Source/Developer/Zen/Private/ZenServerState.cpp
@@ -111,7 +111,7 @@ bool ZenServerState::IsProcessRunning(uint32 Pid)
 		}
 		else if (Error == ERROR_ACCESS_DENIED)
 		{
-			UE_LOG(LogZenServiceState, Warning, TEXT("No access to open running process %d: %d, assuming it is running"), Pid, Error);
+			// This occurs when running as a system service, as we're not the user that started the process.
 			return true;
 		}
 		UE_LOG(LogZenServiceState, Warning, TEXT("Failed to open running process %d: %d, assuming it is not running"), Pid, Error);
diff --git a/Engine/Source/Developer/Zen/Private/ZenServerState.h b/Engine/Source/Developer/Zen/Private/ZenServerState.h
--- a/Engine/Source/Developer/Zen/Private/ZenServerState.h
+++ b/Engine/Source/Developer/Zen/Private/ZenServerState.h
@@ -23,7 +23,8 @@ public:
 		uint8				SessionId[12];
 		std::atomic<uint32> SponsorPids[8];
 		std::atomic<uint16> EffectiveListenPort;
-		uint8				Padding[10];
+		std::atomic<bool>   ActiveSponsors;
+		uint8				Padding[9];
 
 		enum class FlagsEnum : uint16
 		{
diff --git a/Engine/Source/Developer/Zen/Public/Experimental/ZenServerInterface.h b/Engine/Source/Developer/Zen/Public/Experimental/ZenServerInterface.h
--- a/Engine/Source/Developer/Zen/Public/Experimental/ZenServerInterface.h
+++ b/Engine/Source/Developer/Zen/Public/Experimental/ZenServerInterface.h
@@ -66,12 +66,25 @@ struct FServiceAutoLaunchSettings
 	} InstallMode = EInstallMode::Copy;
 };
 
+struct FSystemServiceSettings
+{
+	FString InstallPath;
+	FString DataPath;
+	FString ExtraArgs;
+	uint16 DesiredPort = 8558;
+	bool bSendUnattendedBugReports = false;
+	bool bAllowRemoteNetworkService = false;
+};
+
 struct FServiceSettings
 {
+	bool bSystemService = false;
+	FSystemServiceSettings SystemServiceSettings;
 	TVariant<FServiceAutoLaunchSettings, FServiceConnectSettings> SettingsVariant;
 
-	inline bool IsAutoLaunch() const { return SettingsVariant.IsType<FServiceAutoLaunchSettings>(); }
-	inline bool IsConnectExisting() const { return SettingsVariant.IsType<FServiceConnectSettings>(); }
+	inline bool IsSystemService() const { return bSystemService; }
+	inline bool IsAutoLaunch() const { return IsSystemService() || SettingsVariant.IsType<FServiceAutoLaunchSettings>(); }
+	inline bool IsConnectExisting() const { return !IsSystemService() && SettingsVariant.IsType<FServiceConnectSettings>(); }
 	inline bool IsRequirePlugins() const { return IsAutoLaunch() && !SettingsVariant.Get<FServiceAutoLaunchSettings>().Plugins.IsEmpty(); }
 
 	UE_API bool ReadFromConfig();
@@ -315,6 +328,7 @@ private:
 	bool ConditionalUpdateLocalInstall(FServiceAutoLaunchSettings::EInstallMode InstallMode);
 	bool WriteLocalPluginsConfig();
 	static bool AutoLaunch(const FServiceAutoLaunchSettings& InSettings, const FString& ExecutablePath, const FString& ExecutionContextFilePath, FZenServiceEndpoint& OutEndpoint);
+	bool ConditionalUpdateSystemServiceInstall(const FSystemServiceSettings& InSettings);
 	bool PingService(uint32_t AttemptCount);
 
 	mutable TFuture<TOptional<FGCStatus>> GCStatusRequest;
@@ -339,6 +353,7 @@ private:
 	static uint16 AutoLaunchedPort;
 	static uint32 AutoLaunchedPid;
 	bool bHasLaunchedLocal = false;
+	bool bHasLaunchedSystemService = false;
 	bool bIsRunningLocally = true;
 };
 
