Fix duplicated localization keys in ZenServerInterface.

[RN] none

diff --git a/Engine/Source/Developer/Zen/Private/ZenServerInterface.cpp b/Engine/Source/Developer/Zen/Private/ZenServerInterface.cpp
--- a/Engine/Source/Developer/Zen/Private/ZenServerInterface.cpp
+++ b/Engine/Source/Developer/Zen/Private/ZenServerInterface.cpp
@@ -2618,8 +2618,8 @@ PromptUserUnableToUpdateSystemServiceBinaries(const FString& InstallPath)
 #if !IS_PROGRAM
 	if (!FApp::IsUnattended() && !IsRunningCommandlet() && !GIsRunningUnattendedScript)
 	{
-		FText ZenUnableToUpdateBinariesPromptTitle = NSLOCTEXT("Zen", "Zen_SystemServiceMismatchedConfigurationPromptTitle", "Failed to update Unreal Zen Storage Server binaries");
-		FText ZenUnableToUpdateBinariesPromptText = FText::Format(NSLOCTEXT("Zen", "Zen_SystemServiceMismatchedConfigurationPromptText", "Unreal Zen Storage Server was unable to update the binaries at '%s'. Zen will run using AutoLaunch or ConnectExisting settings."), FText::FromString(InstallPath));
+		FText ZenUnableToUpdateBinariesPromptTitle = NSLOCTEXT("Zen", "Zen_SystemServiceUnableToUpdateBinariesPromptTitle", "Failed to update Unreal Zen Storage Server binaries");
+		FText ZenUnableToUpdateBinariesPromptText = FText::Format(NSLOCTEXT("Zen", "Zen_SystemServiceUnableToUpdateBinariesPromptText", "Unreal Zen Storage Server was unable to update the binaries at '%s'. Zen will run using AutoLaunch or ConnectExisting settings."), FText::FromString(InstallPath));
 		FPlatformMisc::MessageBoxExt(EAppMsgType::Ok, *ZenUnableToUpdateBinariesPromptTitle.ToString(), *ZenUnableToUpdateBinariesPromptText.ToString());
 	}
 	else
