Set zen server install directory to a path that doesn't conflict with our actual install directory during unit tests.

[RN] none

diff --git a/Engine/Source/Developer/DerivedDataCache/Private/Tests/HttpCacheStoreTest.cpp b/Engine/Source/Developer/DerivedDataCache/Private/Tests/HttpCacheStoreTest.cpp
--- a/Engine/Source/Developer/DerivedDataCache/Private/Tests/HttpCacheStoreTest.cpp
+++ b/Engine/Source/Developer/DerivedDataCache/Private/Tests/HttpCacheStoreTest.cpp
@@ -618,6 +618,12 @@ bool FHttpCacheStoreTest::RunTest(const FString& Parameters)
 	ILegacyCacheStore* TestBackend = GetTestBackend();
 
 #if UE_HTTPCACHESTORETEST_USE_ZEN
+	Zen::Private::SetLocalInstallPathOverride(*FPaths::ConvertRelativePathToFull(FPaths::Combine(FPaths::EngineSavedDir(), "ZenUnitTestInstall")));
+	ON_SCOPE_EXIT
+	{
+		Zen::Private::SetLocalInstallPathOverride(TEXT(""));
+	};
+
 	using namespace UE::Zen;
 	FServiceSettings ZenUpstreamTestServiceSettings;
 	FServiceAutoLaunchSettings& ZenUpstreamTestAutoLaunchSettings = ZenUpstreamTestServiceSettings.SettingsVariant.Get<FServiceAutoLaunchSettings>();
diff --git a/Engine/Source/Developer/Zen/Private/ZenServerInterface.cpp b/Engine/Source/Developer/Zen/Private/ZenServerInterface.cpp
--- a/Engine/Source/Developer/Zen/Private/ZenServerInterface.cpp
+++ b/Engine/Source/Developer/Zen/Private/ZenServerInterface.cpp
@@ -208,10 +208,30 @@ GetSystemServiceVersionCachePath(const FSystemServiceSettings& InSettings)
 	return InstallVersionCache;
 }
 
+namespace Private
+{
+// For testing purposes. See comments on GetLocalInstallPathOverride() in ZenServerInterface.h.
+static FString LocalInstallPathOverride;
+
+void
+SetLocalInstallPathOverride(const FString& Path)
+{
+	checkf(LocalInstallPathOverride.IsEmpty() || Path.IsEmpty(), TEXT("Local Zen install path override already set to %s"), *LocalInstallPathOverride);
+	LocalInstallPathOverride = Path;
+}
+
+FString
+GetLocalInstallPathOverride()
+{
+	return LocalInstallPathOverride;
+}
+
+}
+
 static FString
 GetLocalInstallPath()
 {
-	return FPaths::ConvertRelativePathToFull(FPaths::Combine(GetLocalZenRootPath(), TEXT("Zen\\Install")));
+	return Private::GetLocalInstallPathOverride().IsEmpty() ? FPaths::ConvertRelativePathToFull(FPaths::Combine(GetLocalZenRootPath(), TEXT("Zen\\Install"))) : Private::GetLocalInstallPathOverride();
 }
 
 static FString
diff --git a/Engine/Source/Developer/Zen/Public/Experimental/ZenServerInterface.h b/Engine/Source/Developer/Zen/Public/Experimental/ZenServerInterface.h
--- a/Engine/Source/Developer/Zen/Public/Experimental/ZenServerInterface.h
+++ b/Engine/Source/Developer/Zen/Public/Experimental/ZenServerInterface.h
@@ -171,6 +171,15 @@ namespace Private
 {
 	UE_API bool IsLocalAutoLaunched(FStringView InstanceURL);
 	UE_API bool GetLocalDataCachePathOverride(FString& OutDataPath);
+
+	// For testing purposes only.
+	//
+	// This is a bit of a hack, to avoid issues in unit tests where we attempt to remove a copied binary when
+	// testing install in link mode, or vice versa. During unit tests, we need to run those installations sandboxed
+	// to their own install directory so as not to interfere with an actual running zenserver that is being
+	// used outside of the unit tests.
+	UE_API FString GetLocalInstallPathOverride();
+	UE_API void SetLocalInstallPathOverride(const FString& Path);
 }
 
 class FZenServiceInstance;
diff --git a/Engine/Source/Developer/Zen/Tests/ZenTest.cpp b/Engine/Source/Developer/Zen/Tests/ZenTest.cpp
--- a/Engine/Source/Developer/Zen/Tests/ZenTest.cpp
+++ b/Engine/Source/Developer/Zen/Tests/ZenTest.cpp
@@ -23,6 +23,12 @@ TEST_CASE("Zen::ZenServerInterface", "[Zen][Basic]")
 	FString DefaultDataPath = FPaths::ConvertRelativePathToFull(FPaths::Combine(DataPathRoot, "Default"));
 	IFileManager::Get().DeleteDirectory(*DataPathRoot, false, true);
 
+	Zen::Private::SetLocalInstallPathOverride(DataPathRoot);
+	ON_SCOPE_EXIT
+	{
+		Zen::Private::SetLocalInstallPathOverride(TEXT(""));
+	};
+
 	SECTION("Basic AutoLaunch and Shutdown (copy)")
 	{
 		for (int Iteration = 0; Iteration < 4; ++Iteration)
