Fix to Zen Server calling CreateProc wrong

diff --git a/Engine/Source/Developer/Zen/Private/ZenServerInterface.cpp b/Engine/Source/Developer/Zen/Private/ZenServerInterface.cpp
--- a/Engine/Source/Developer/Zen/Private/ZenServerInterface.cpp
+++ b/Engine/Source/Developer/Zen/Private/ZenServerInterface.cpp
@@ -365,8 +365,7 @@ static bool RunZenUtility(const FString& UtilityPath, const FString& CommandLine
 		OutProcessID,
 		PriorityModifier,
 		nullptr,
-		PipeWriteChild,
-		PipeReadChild);
+		PipeWriteChild);
 
 	if (!ZenUtilProc.IsValid())
 	{
