Increase timeout when waiting for other process to install/start zenserver

#rnx

diff --git a/Engine/Source/Developer/Zen/Private/ZenServerInterface.cpp b/Engine/Source/Developer/Zen/Private/ZenServerInterface.cpp
--- a/Engine/Source/Developer/Zen/Private/ZenServerInterface.cpp
+++ b/Engine/Source/Developer/Zen/Private/ZenServerInterface.cpp
@@ -375,7 +375,7 @@ static bool RunZenUtility(const FString& UtilityPath, const FString& CommandLine
 
 	bool DumpWritten = false;
 
-	const double MaxWaitDuration = FApp::IsUnattended() ? 15.0 : 5.0;
+	const double MaxWaitDuration = FApp::IsUnattended() ? 25.0 : 8.0;
 
 	const uint64 StartTime = FPlatformTime::Cycles64();
 	while (FPlatformProcess::IsProcRunning(ZenUtilProc))
@@ -383,7 +383,7 @@ static bool RunZenUtility(const FString& UtilityPath, const FString& CommandLine
 		double Duration = FPlatformTime::ToSeconds64(FPlatformTime::Cycles64() - StartTime);
 		if (Duration >= 2.0)
 		{
-			if (Duration > 5.0)
+			if (Duration > 15.0)
 			{
 #if PLATFORM_WINDOWS
 				if (!DumpWritten)
@@ -2404,7 +2404,7 @@ FZenServiceInstance::Initialize()
 
 	if (!bSystemService)
 	{
-		FSystemWideCriticalSection InstallCriticalSection(TEXT("ZenServerInstall"), FTimespan::FromSeconds(10));
+		FSystemWideCriticalSection InstallCriticalSection(TEXT("ZenServerInstall"), FTimespan::FromSeconds(FApp::IsUnattended() ? 30.0 : 10.0));
 		if (!InstallCriticalSection.IsValid())
 		{
 			PromptUserOfFailedSystemServiceUninstall();
@@ -2427,16 +2427,17 @@ FZenServiceInstance::Initialize()
 			const FString ExecutionContextFilePath = GetServiceRunContextPath();
 
 			int LaunchAttempts = 0;
-			const FTimespan MaximumWaitForHealth = FTimespan::FromSeconds(20);
+			const FTimespan MaximumWaitForHealth = FTimespan::FromSeconds(30);
 
 			FDateTime StartedWaitingForHealth = FDateTime::UtcNow();
 			while (true)
 			{
 				{
-					FSystemWideCriticalSection RecoveryCriticalSection(TEXT("ZenServerLaunch"), FTimespan::FromSeconds(10));
+					FSystemWideCriticalSection RecoveryCriticalSection(TEXT("ZenServerLaunch"), FTimespan::FromSeconds(FApp::IsUnattended() ? 15.0 : 5.0));
 					if (!RecoveryCriticalSection.IsValid())
 					{
-						UE_LOG(LogZenServiceInstance, Warning, TEXT("Local ZenServer AutoLaunch initialization timed out waiting for other process to complete startup/recovery"));
+						FTimespan WaitForRecovery = FDateTime::UtcNow() - StartedWaitingForHealth;
+						UE_LOG(LogZenServiceInstance, Display, TEXT("Local ZenServer AutoLaunch initialization still waiting for other process to complete startup/recovery, waited %.3lf seconds"), WaitForRecovery.GetTotalSeconds());
 					}
 					else
 					{
@@ -2463,14 +2464,14 @@ FZenServiceInstance::Initialize()
 
 				++LaunchAttempts;
 				FTimespan WaitForHealth = FDateTime::UtcNow() - StartedWaitingForHealth;
-				if ((WaitForHealth > MaximumWaitForHealth) && (LaunchAttempts > 1))
+				if ((WaitForHealth > MaximumWaitForHealth) && (LaunchAttempts > 2))
 				{
 					bHasLaunchedLocal = false;
 					bIsRunningLocally = false;
-					UE_LOG(LogZenServiceInstance, Warning, TEXT("Local ZenServer AutoLaunch initialization timed out waiting for service to become healthy"));
+					UE_LOG(LogZenServiceInstance, Warning, TEXT("Local ZenServer AutoLaunch initialization timed out waiting for service to become healthy, waited %.3lf seconds"), WaitForHealth.GetTotalSeconds());
 					break;
 				}
-				UE_LOG(LogZenServiceInstance, Log, TEXT("Awaiting ZenServer readiness"));
+				UE_LOG(LogZenServiceInstance, Log, TEXT("Awaiting ZenServer readiness, waited for %.3lf seconds"), WaitForHealth.GetTotalSeconds());
 				FPlatformProcess::Sleep(0.5f);
 			}
 		}
@@ -2846,7 +2847,7 @@ bool FZenServiceInstance::ConditionalUpdateSystemServiceInstall(const FSystemSer
 	FString Output;
 	int ReturnCode;
 
-	FSystemWideCriticalSection InstallCriticalSection(TEXT("ZenServerInstall"), FTimespan::FromSeconds(10));
+	FSystemWideCriticalSection InstallCriticalSection(TEXT("ZenServerInstall"), FTimespan::FromSeconds(FApp::IsUnattended() ? 30.0 : 10.0));
 	if (!InstallCriticalSection.IsValid())
 	{
 		PromptUserOfFailedSystemServiceInstall();
