Avoid warnings about uninstalling system services when working with a Zen version that is not aware of service commands.

[RN] none

diff --git a/Engine/Source/Developer/Zen/Private/ZenServerInterface.cpp b/Engine/Source/Developer/Zen/Private/ZenServerInterface.cpp
--- a/Engine/Source/Developer/Zen/Private/ZenServerInterface.cpp
+++ b/Engine/Source/Developer/Zen/Private/ZenServerInterface.cpp
@@ -2840,6 +2840,13 @@ bool UninstallSystemService(FString UtilityPath)
 		return false;
 	}
 
+	// Versions of zen that don't know about the service command will still return 0 for 'service status', leading to warnings
+	// later on in this code when we try to uninstall.
+	if (Output.StartsWith("Unknown command"))
+	{
+		return true;
+	}
+
 	// zen returns 1 for status of a service that's not installed
 	if (ReturnCode == 1)
 	{
