Added UI for retrying canceled or failed build transfers.  Retrying will enqueue a new transfer with the same parameters as a past one to the end of the queue.

diff --git a/Engine/Source/Developer/Zen/Private/BuildServerInterface.cpp b/Engine/Source/Developer/Zen/Private/BuildServerInterface.cpp
--- a/Engine/Source/Developer/Zen/Private/BuildServerInterface.cpp
+++ b/Engine/Source/Developer/Zen/Private/BuildServerInterface.cpp
@@ -412,7 +412,6 @@ FBuildServiceInstance::ListBuilds(FStringView Namespace, FStringView Bucket, FOn
 FBuildServiceInstance::FBuildTransfer
 FBuildServiceInstance::StartBuildTransfer(const FCbObjectId& BuildId, FStringView Name, FStringView DestinationFolder, FStringView Namespace, FStringView Bucket)
 {
-	const FString ZenUtility = UE::Zen::GetLocalInstallUtilityPath();
 	TSharedPtr<Private::FBuildTransferState> NewState = MakeShared<Private::FBuildTransferState>();
 	NewState->Type = Private::EBuildTransferType::Files;
 	NewState->BuildId = BuildId;
@@ -431,7 +430,6 @@ FBuildServiceInstance::StartBuildTransfer(const FCbObjectId& BuildId, FStringVie
 FBuildServiceInstance::FBuildTransfer
 FBuildServiceInstance::StartOplogBuildTransfer(const FCbObjectId& BuildId, FStringView Name, FStringView DestinationProjectId, FStringView DestinationOplogId, FStringView ProjectFilePath, FStringView Namespace, FStringView Bucket, FName TargetPlatformName)
 {
-	const FString ZenUtility = UE::Zen::GetLocalInstallUtilityPath();
 	TSharedPtr<Private::FBuildTransferState> NewState = MakeShared<Private::FBuildTransferState>();
 	NewState->Type = Private::EBuildTransferType::Oplog;
 	NewState->BuildId = BuildId;
@@ -449,6 +447,31 @@ FBuildServiceInstance::StartOplogBuildTransfer(const FCbObjectId& BuildId, FStri
 	return NewTransfer;
 }
 
+FBuildServiceInstance::FBuildTransfer
+FBuildServiceInstance::RepeatBuildTransfer(FBuildTransfer BuildTransfer)
+{
+	TSharedPtr<Private::FBuildTransferState> NewState = MakeShared<Private::FBuildTransferState>();
+
+	{
+		FReadScopeLock _(BuildTransfer.State->Lock);
+		NewState->Type = BuildTransfer.State->Type;
+		NewState->BuildId = BuildTransfer.State->BuildId;
+		NewState->Name = BuildTransfer.State->Name;
+		NewState->Namespace = BuildTransfer.State->Namespace;
+		NewState->Bucket = BuildTransfer.State->Bucket;
+		NewState->Destination = BuildTransfer.State->Destination;
+		NewState->ProjectFilePath = BuildTransfer.State->ProjectFilePath;
+		NewState->TargetPlatformName = BuildTransfer.State->TargetPlatformName;
+	}
+	
+	BuildTransferThreadStates.Enqueue(NewState.ToSharedRef());
+	KickBuildTransferThread();
+
+	FBuildTransfer NewTransfer;
+	NewTransfer.State = NewState;
+	return NewTransfer;
+}
+
 FString
 FBuildServiceInstance::FBuildTransfer::GetDescription() const
 {
diff --git a/Engine/Source/Developer/Zen/Public/Experimental/BuildServerInterface.h b/Engine/Source/Developer/Zen/Public/Experimental/BuildServerInterface.h
--- a/Engine/Source/Developer/Zen/Public/Experimental/BuildServerInterface.h
+++ b/Engine/Source/Developer/Zen/Public/Experimental/BuildServerInterface.h
@@ -111,6 +111,7 @@ public:
 
 	UE_API FBuildTransfer StartBuildTransfer(const FCbObjectId& BuildId, FStringView Name, FStringView DestinationFolder, FStringView Namespace, FStringView Bucket);
 	UE_API FBuildTransfer StartOplogBuildTransfer(const FCbObjectId& BuildId, FStringView Name, FStringView DestinationProjectId, FStringView DestinationOplogId, FStringView ProjectFilePath, FStringView Namespace, FStringView Bucket, FName TargetPlatformName);
+	UE_API FBuildTransfer RepeatBuildTransfer(FBuildTransfer ExistingTransfer);
 private:
 	typedef TUniqueFunction<void(FString& CommandLineArgs)> FPreZenUtilityInvocation;
 	typedef TUniqueFunction<void(bool)> FOnZenUtilityInvocationComplete;
diff --git a/Engine/Source/Editor/StorageServerWidgets/Internal/SBuildActivity.h b/Engine/Source/Editor/StorageServerWidgets/Internal/SBuildActivity.h
--- a/Engine/Source/Editor/StorageServerWidgets/Internal/SBuildActivity.h
+++ b/Engine/Source/Editor/StorageServerWidgets/Internal/SBuildActivity.h
@@ -53,6 +53,8 @@ private:
 
 	UE_API bool CanCancelBuildActivity() const;
 	UE_API void CancelBuildActivity() const;
+	UE_API bool CanRetryBuildActivity() const;
+	UE_API void RetryBuildActivity();
 	UE_API bool CanClearBuildActivity() const;
 	UE_API void ClearBuildActivity();
 	UE_API bool CanViewLogForBuildActivity() const;
diff --git a/Engine/Source/Editor/StorageServerWidgets/Private/SBuildActivity.cpp b/Engine/Source/Editor/StorageServerWidgets/Private/SBuildActivity.cpp
--- a/Engine/Source/Editor/StorageServerWidgets/Private/SBuildActivity.cpp
+++ b/Engine/Source/Editor/StorageServerWidgets/Private/SBuildActivity.cpp
@@ -152,6 +152,18 @@ TSharedPtr<SWidget> SBuildActivity::OnGetBuildActivityContextMenuContent()
 		EUserInterfaceActionType::Button
 	);
 
+	MenuBuilder.AddMenuEntry(
+		LOCTEXT("BuildActivity_Retry", "Retry"),
+		LOCTEXT("BuildActivity_RetryTooltip", "Retry the failed or canceled activity"),
+		FSlateIcon(),
+		FUIAction(
+			FExecuteAction::CreateRaw(this, &SBuildActivity::RetryBuildActivity),
+			FCanExecuteAction::CreateRaw(this, &SBuildActivity::CanRetryBuildActivity)
+		),
+		NAME_None,
+		EUserInterfaceActionType::Button
+	);
+
 	MenuBuilder.AddMenuEntry(
 		LOCTEXT("BuildActivity_Clear", "Clear"),
 		LOCTEXT("BuildActivity_ClearTooltip", "Clear this activity if it is completed"),
@@ -203,6 +215,35 @@ void SBuildActivity::CancelBuildActivity() const
 	SelectedItems[0]->Transfer.RequestCancel();
 }
 
+bool SBuildActivity::CanRetryBuildActivity() const
+{
+	using namespace UE::Zen::Build;
+
+	TArray<FBuildActivityPtr> SelectedItems = BuildActivityListView->GetSelectedItems();
+	if (SelectedItems.Num() != 1)
+	{
+		return false;
+	}
+	FBuildServiceInstance::EBuildTransferStatus Status = SelectedItems[0]->Transfer.GetStatus();
+	return (Status == FBuildServiceInstance::EBuildTransferStatus::Canceled) || (Status == FBuildServiceInstance::EBuildTransferStatus::Failed);
+}
+
+void SBuildActivity::RetryBuildActivity()
+{
+	using namespace UE::Zen::Build;
+
+	TArray<FBuildActivityPtr> SelectedItems = BuildActivityListView->GetSelectedItems();
+	if (SelectedItems.Num() != 1)
+	{
+		return;
+	}
+	if (TSharedPtr<FBuildServiceInstance> ServiceInstance = BuildServiceInstance.Get())
+	{
+		AddBuildTransfer(ServiceInstance->RepeatBuildTransfer(SelectedItems[0]->Transfer),
+			SelectedItems[0]->Name, SelectedItems[0]->Platform);
+	}
+}
+
 bool SBuildActivity::CanClearBuildActivity() const
 {
 	using namespace UE::Zen::Build;
