[ProjectLauncher2] Change build retrieving logic to use BuildServerInterface and match the logic the BuildStorageTool uses
- Added snapshot changelist display to custom profiles
- Removed build list logic from BackendLaunchExtension and move it to a shared place in the launcher profile

diff --git a/Engine/Plugins/Developer/ProjectLauncher/Source/ProjectLauncher/Private/ProfileTree/CustomProfileTreeBuilder.cpp b/Engine/Plugins/Developer/ProjectLauncher/Source/ProjectLauncher/Private/ProfileTree/CustomProfileTreeBuilder.cpp
--- a/Engine/Plugins/Developer/ProjectLauncher/Source/ProjectLauncher/Private/ProfileTree/CustomProfileTreeBuilder.cpp
+++ b/Engine/Plugins/Developer/ProjectLauncher/Source/ProjectLauncher/Private/ProfileTree/CustomProfileTreeBuilder.cpp
@@ -27,6 +27,7 @@ namespace ProjectLauncher
 		AddCompressPakFilesProperty(ContentSchemeHeader);
 		AddUseIoStoreProperty(ContentSchemeHeader);
 		AddGenerateChunksProperty(ContentSchemeHeader);
+		AddZenSnapshotProperty(ContentSchemeHeader);
 		AddImportZenSnapshotProperty(ContentSchemeHeader);
 		AddZenPakStreamingPathProperty(ContentSchemeHeader);
 
diff --git a/Engine/Plugins/Developer/ProjectLauncher/Source/ProjectLauncher/Private/ProfileTree/GenericProfileTreeBuilder.cpp b/Engine/Plugins/Developer/ProjectLauncher/Source/ProjectLauncher/Private/ProfileTree/GenericProfileTreeBuilder.cpp
--- a/Engine/Plugins/Developer/ProjectLauncher/Source/ProjectLauncher/Private/ProfileTree/GenericProfileTreeBuilder.cpp
+++ b/Engine/Plugins/Developer/ProjectLauncher/Source/ProjectLauncher/Private/ProfileTree/GenericProfileTreeBuilder.cpp
@@ -272,6 +272,23 @@ namespace ProjectLauncher
 		);
 	}
 
+	void FGenericProfileTreeBuilder::AddZenSnapshotProperty(FLaunchProfileTreeNode& HeadingNode)
+	{
+		HeadingNode.AddWidget(LOCTEXT("ZenSnapshotLabel", "Closest Zen Snapshot"),
+			{
+				.IsVisible = [this]() { return ForContent(); }
+			},
+			SNew(STextBlock)
+			.Text(this, &FGenericProfileTreeBuilder::GetZenSnapshotText)
+			.Font(FCoreStyle::Get().GetFontStyle("SmallFont")));
+	}
+
+	FText FGenericProfileTreeBuilder::GetZenSnapshotText() const
+	{
+		int32 Build = Profile->GetZenSnapshot();
+		return FText::FromString(FString::FromInt(Build));
+	}
+
 	void FGenericProfileTreeBuilder::AddImportZenSnapshotProperty( FLaunchProfileTreeNode& HeadingNode )
 	{
 		HeadingNode.AddBoolean( LOCTEXT("ImportZenSnapshotLabel", "Import Best Match Zen Snapshot"),
diff --git a/Engine/Plugins/Developer/ProjectLauncher/Source/ProjectLauncher/Public/ProfileTree/GenericProfileTreeBuilder.h b/Engine/Plugins/Developer/ProjectLauncher/Source/ProjectLauncher/Public/ProfileTree/GenericProfileTreeBuilder.h
--- a/Engine/Plugins/Developer/ProjectLauncher/Source/ProjectLauncher/Public/ProfileTree/GenericProfileTreeBuilder.h
+++ b/Engine/Plugins/Developer/ProjectLauncher/Source/ProjectLauncher/Public/ProfileTree/GenericProfileTreeBuilder.h
@@ -52,6 +52,7 @@ namespace ProjectLauncher
 		UE_API void AddCompressPakFilesProperty( FLaunchProfileTreeNode& HeadingNode );
 		UE_API void AddUseIoStoreProperty( FLaunchProfileTreeNode& HeadingNode );
 		UE_API void AddGenerateChunksProperty( FLaunchProfileTreeNode& HeadingNode );
+		UE_API void AddZenSnapshotProperty(FLaunchProfileTreeNode& HeadingNode);
 		UE_API void AddImportZenSnapshotProperty( FLaunchProfileTreeNode& HeadingNode );
 		UE_API void AddZenPakStreamingPathProperty( FLaunchProfileTreeNode& HeadingNode );
 		UE_API void AddIncrementalCookProperty( FLaunchProfileTreeNode& HeadingNode );
@@ -176,7 +177,7 @@ namespace ProjectLauncher
 		UE_API void CacheStartupMapList() const;
 		UE_API void CacheArchitectures();
 		UE_API TSharedRef<SWidget> OnGenerateComboWidget( TSharedPtr<FString> InComboString );
-
+		FText GetZenSnapshotText() const;
 
 		FLaunchProfileTreeDataRef TreeData;
 		const ILauncherProfileRef Profile;
diff --git a/Engine/Source/Developer/LauncherServices/LauncherServices.Build.cs b/Engine/Source/Developer/LauncherServices/LauncherServices.Build.cs
--- a/Engine/Source/Developer/LauncherServices/LauncherServices.Build.cs
+++ b/Engine/Source/Developer/LauncherServices/LauncherServices.Build.cs
@@ -26,6 +26,7 @@ namespace UnrealBuildTool.Rules
 					"UnrealEdMessages",
 					"Json",
 					"TurnkeyIO",
+					"Zen"
 				});
 
 			PrivateIncludePathModuleNames.AddRange(
diff --git a/Engine/Source/Developer/LauncherServices/Private/Profiles/LauncherProfile.cpp b/Engine/Source/Developer/LauncherServices/Private/Profiles/LauncherProfile.cpp
--- a/Engine/Source/Developer/LauncherServices/Private/Profiles/LauncherProfile.cpp
+++ b/Engine/Source/Developer/LauncherServices/Private/Profiles/LauncherProfile.cpp
@@ -76,5 +76,80 @@ FString LexToStringLocalized(ELauncherProfileValidationErrors::Type Value)
 	};
 }
 
+TSharedPtr<UE::Zen::Build::FBuildListRetriever> FLauncherProfile::BuildListRetriever;
+TMap<FString, FLauncherProfile::FPlatformToBuildsMap> FLauncherProfile::PerProjectBuilds;
+
+void FLauncherProfile::QueryZenSnapshotsForProject()
+{
+	FString BuildType = TEXT("oplog");
+	FString Branch = FEngineVersion::Current().GetBranch();
+	Branch = Branch.Replace(TEXT("//"), TEXT("")).Replace(TEXT("/"), TEXT("-")).ToLower();
+
+	TSharedPtr<UE::Zen::Build::FBuildListRetriever> BuildRetriever = FLauncherProfile::GetBuildListRetriever();
+	BuildRetriever->QueryBuilds(GetProjectName(), BuildType, Branch,
+		[this](const TMap<FString, TArray<int32>>& PerPlatformBuilds) mutable
+		{
+			check(IsInGameThread());
+
+			FPlatformToBuildsMap& PlatformConfigToBuildMap = PerProjectBuilds.FindOrAdd(GetProjectName().ToLower());
+			PlatformConfigToBuildMap.Reset();
+
+			for (auto& Pair : PerPlatformBuilds)
+			{
+				// Find the closest changelist that is before the synced changelist
+				int32 ClosestBuild = -1;
+				const int32 ChangeListToMatch = FEngineVersion::Current().GetChangelist();
+				for (int32 Build : Pair.Value)
+				{
+					if (Build <= ChangeListToMatch)
+					{
+						ClosestBuild = Build;
+					}
+					else
+					{
+						break;
+					}
+				}
+
+				PlatformConfigToBuildMap.Add(Pair.Key, ClosestBuild);
+			}
+		}
+	);
+}
+
+int32 FLauncherProfile::GetZenSnapshot()
+{
+	FString ProjectName = GetProjectName().ToLower();
+	if (PerProjectBuilds.Contains(ProjectName))
+	{
+		TArray<FString> Platforms = GetCookedPlatforms();
+		if (Platforms.Num() > 0)
+		{
+			FString Platform = Platforms[0].ToLower();
+			FPlatformToBuildsMap& Builds = PerProjectBuilds[ProjectName];
+			if (Builds.Contains(Platform))
+			{
+				return Builds[Platform];
+			}
+		}
+	}
+	else
+	{
+		QueryZenSnapshotsForProject();
+	}
+
+	return 0;
+}
+
+TSharedPtr<UE::Zen::Build::FBuildListRetriever> FLauncherProfile::GetBuildListRetriever()
+{
+	if (!BuildListRetriever)
+	{
+		BuildListRetriever = MakeShared<UE::Zen::Build::FBuildListRetriever>();
+		BuildListRetriever->ConnectToBuildService();
+	}
+	return BuildListRetriever;
+}
+
 #undef LOCTEXT_NAMESPACE
 
diff --git a/Engine/Source/Developer/LauncherServices/Private/Profiles/LauncherProfile.h b/Engine/Source/Developer/LauncherServices/Private/Profiles/LauncherProfile.h
--- a/Engine/Source/Developer/LauncherServices/Private/Profiles/LauncherProfile.h
+++ b/Engine/Source/Developer/LauncherServices/Private/Profiles/LauncherProfile.h
@@ -23,6 +23,7 @@
 #include "TargetReceipt.h"
 #include "DesktopPlatformModule.h"
 #include "GameProjectHelper.h"
+#include "Experimental/BuildListRetriever.h"
 
 DEFINE_LOG_CATEGORY_STATIC(LogLauncherProfile, Log, All);
 
@@ -232,6 +233,7 @@ private:
  */
 class FLauncherProfile final
 	: public ILauncherProfile
+	, public TSharedFromThis<FLauncherProfile>
 {
 public:
 
@@ -2874,6 +2876,8 @@ public:
 		return bUseZenStore;
 	}
 
+	virtual int32 GetZenSnapshot() override;
+
 	virtual void SetImportingZenSnapshot( bool Import ) override
 	{
 		if (bImportingZenSnapshot != Import)
@@ -3460,6 +3464,8 @@ protected:
 	}
 
 private:
+	void QueryZenSnapshotsForProject();
+	TSharedPtr<UE::Zen::Build::FBuildListRetriever> GetBuildListRetriever();
 
 	//  Holds a reference to the launcher profile manager.
 	ILauncherProfileManagerRef LauncherProfileManager;
@@ -3722,4 +3728,11 @@ private:
 
 	// Cached build target options (not serialized)
 	TArray<FString> ExplictBuildTargetNames;
+
+	// Utility class instance used to query build list
+	static TSharedPtr<UE::Zen::Build::FBuildListRetriever> BuildListRetriever;
+
+	// Cache all builds for projects and platforms, shared between profiles
+	typedef TMap<FString, int32> FPlatformToBuildsMap;
+	static TMap<FString, FPlatformToBuildsMap> PerProjectBuilds;
 };
diff --git a/Engine/Source/Developer/LauncherServices/Public/ILauncherProfile.h b/Engine/Source/Developer/LauncherServices/Public/ILauncherProfile.h
--- a/Engine/Source/Developer/LauncherServices/Public/ILauncherProfile.h
+++ b/Engine/Source/Developer/LauncherServices/Public/ILauncherProfile.h
@@ -1505,6 +1505,13 @@ public:
 	 */
 	virtual bool IsUsingZenStore() const = 0;
 
+	/**
+	 * Get the Zen snapshot build number used for import.
+	 *
+	 * @return Zen snapshot build number
+	 */
+	virtual int32 GetZenSnapshot() = 0;
+
 	/**
 	 * Set whether to import a Zen snapshot before cooking.
 	 * 
diff --git a/Engine/Source/Developer/Zen/Internal/ZenBuildUtils.h b/Engine/Source/Developer/Zen/Internal/ZenBuildUtils.h
new file mode 100644
--- /dev/null
+++ b/Engine/Source/Developer/Zen/Internal/ZenBuildUtils.h
@@ -0,0 +1,29 @@
+// Copyright Epic Games, Inc. All Rights Reserved.
+
+#pragma once
+
+#include "Experimental/BuildServerInterface.h"
+
+namespace UE::Zen::Build
+{
+	struct FNamespacePlatformBucketTuple
+	{
+		FString Namespace;
+		FString Platform;
+		FString Bucket;
+	};
+
+	struct FBuildState
+	{
+		FString Namespace;
+		FString Platform;
+		TArray<UE::Zen::Build::FBuildServiceInstance::FBuildRecord> Results;
+	};
+
+	struct FListBuildsState
+	{
+		std::atomic<uint32> PendingQueries;
+		TArray<FBuildState> QueryState;
+	};
+} // namespace UE::Zen::Build
+
diff --git a/Engine/Source/Developer/Zen/Private/BuildListRetriever.cpp b/Engine/Source/Developer/Zen/Private/BuildListRetriever.cpp
new file mode 100644
--- /dev/null
+++ b/Engine/Source/Developer/Zen/Private/BuildListRetriever.cpp
@@ -0,0 +1,200 @@
+// Copyright Epic Games, Inc. All Rights Reserved.
+
+#include "Experimental/BuildListRetriever.h"
+#include "Misc/DateTime.h"
+#include "Misc/EngineVersion.h"
+
+namespace UE::Zen::Build
+{
+
+FBuildListRetriever::FBuildListRetriever()
+{
+	ServiceInstance = MakeShared<FBuildServiceInstance>();
+}
+
+void FBuildListRetriever::ConnectToBuildService()
+{
+	if (bConnected)
+		return;
+
+	bConnected = true;
+	ServiceInstance->OnRefreshNamespacesAndBucketsComplete().AddSP(this, &FBuildListRetriever::OnNamespacesAndBucketsRefreshed);
+	ServiceInstance->Connect([this](UE::Zen::Build::FBuildServiceInstance::EConnectionState ConnectionState)
+		{
+			if (ConnectionState == UE::Zen::Build::FBuildServiceInstance::EConnectionState::ConnectionSucceeded)
+			{
+				ServiceInstance->RefreshNamespacesAndBuckets();
+			}
+		});
+}
+
+void FBuildListRetriever::OnNamespacesAndBucketsRefreshed()
+{
+	ExecuteOnGameThread(UE_SOURCE_LOCATION,
+		[this]
+		{
+			ReadyToIssueBuildQuery = true;
+		});
+}
+
+bool FBuildListRetriever::QueryBuilds(const FString& InProjectName, const FString& InBuildType, const FString& InStream, FOnQueriesComplete&& InOnQueriesCompleteCallback)
+{
+	check(IsInGameThread());
+	if (ReadyToIssueBuildQuery)
+	{
+		ProjectName = InProjectName;
+		BuildType = InBuildType;
+		Stream = InStream;
+		OnQueriesCompleteCallback = MoveTemp(InOnQueriesCompleteCallback);
+
+		IssueBuildQueries();
+		return true;
+	}
+	else
+	{
+		// Don't do anything if an existing query is in flight
+		return false;
+	}
+}
+
+void FBuildListRetriever::IssueBuildQueries()
+{
+	check(IsInGameThread());
+	ReadyToIssueBuildQuery = false;
+
+	TMultiMap<FString, FString> NamespacesAndBuckets = ServiceInstance->GetNamespacesAndBuckets();
+	TMultiMap<FString, FString> BucketsToNamespaces;
+	for (const TPair<FString, FString>& NamespaceAndBucket : NamespacesAndBuckets)
+	{
+		BucketsToNamespaces.AddUnique(NamespaceAndBucket.Value, NamespaceAndBucket.Key);
+	}
+	
+	auto StringToSegmentViews = [](const FString& Str, TArray<FStringView>& OutViews)
+		{
+			FStringView WorkingStringView(Str);
+			int32 CurrentIndex = 0;
+			while (WorkingStringView.FindChar(TCHAR('.'), CurrentIndex))
+			{
+				if (CurrentIndex != 0)
+				{
+					OutViews.Add(WorkingStringView.Left(CurrentIndex));
+				}
+				WorkingStringView.RightChopInline(CurrentIndex + 1);
+			}
+			if (!WorkingStringView.IsEmpty())
+			{
+				OutViews.Add(WorkingStringView);
+			}
+		};
+
+	TMultiMap<FStringView, TArray<FStringView>> NamespacesToBucketSegmentViews;
+	for (const TPair<FString, FString>& NamespaceAndBucket : NamespacesAndBuckets)
+	{
+		BucketsToNamespaces.AddUnique(NamespaceAndBucket.Value, NamespaceAndBucket.Key);
+		TArray<FStringView>& BucketSegmentViews = NamespacesToBucketSegmentViews.Add(NamespaceAndBucket.Key);
+		StringToSegmentViews(NamespaceAndBucket.Value, BucketSegmentViews);
+	}
+
+	const uint32 SegmentIndexProject = 0;
+	const uint32 SegmentIndexBuildType = 1;
+	const uint32 SegmentIndexStream = 2;
+	const uint32 SegmentIndexPlatform = 3;
+	const uint32 SegmentIndexNum = 4;
+
+	// Platform list generation
+	TArray<FString> Platforms;
+	for (const TPair<FStringView, TArray<FStringView>>& NamespaceToBucketSegmentViews : NamespacesToBucketSegmentViews)
+	{
+		const TArray<FStringView>& BucketSegmentViews = NamespaceToBucketSegmentViews.Value;
+		if (BucketSegmentViews.Num() == SegmentIndexNum)
+		{
+			if (BucketSegmentViews[SegmentIndexStream] != Stream)
+			{
+				continue;
+			}
+
+			if (BucketSegmentViews[SegmentIndexProject] != ProjectName)
+			{
+				continue;
+			}
+
+			if (BucketSegmentViews[SegmentIndexBuildType] != BuildType)
+			{
+				continue;
+			}
+
+			Platforms.AddUnique(FString(BucketSegmentViews[SegmentIndexPlatform]));
+		}
+	}
+
+	TArray<FNamespacePlatformBucketTuple> NamespacePlatformBucketTuples;
+	for (const FString& Platform : Platforms)
+	{
+		FString Bucket = FString::Printf(TEXT("%s.%s.%s.%s"), *ProjectName, *BuildType, *Stream, *Platform);
+		// TODO: replace with passed in namespace
+		TArray<FString> NamespacesForBucket;
+		BucketsToNamespaces.MultiFind(Bucket, NamespacesForBucket);
+		for (FString& Namespace2 : NamespacesForBucket)
+		{
+			NamespacePlatformBucketTuples.Emplace(MoveTemp(Namespace2), *Platform, Bucket);
+		}
+	}
+
+	TSharedPtr<FListBuildsState> PendingQueryState = MakeShared<FListBuildsState>();
+	PendingQueryState->PendingQueries = NamespacePlatformBucketTuples.Num();
+	PendingQueryState->QueryState.SetNum(NamespacePlatformBucketTuples.Num());
+	uint32 QueryIndex = 0;
+
+	for (const FNamespacePlatformBucketTuple& NamespacePlatformBucket : NamespacePlatformBucketTuples)
+	{
+		ServiceInstance->ListBuilds(NamespacePlatformBucket.Namespace, NamespacePlatformBucket.Bucket,
+			[this, QueryIndex, Namespace = NamespacePlatformBucket.Namespace, Platform = NamespacePlatformBucket.Platform,
+			PendingQueryState]
+			(TArray<FBuildServiceInstance::FBuildRecord>&& Results) mutable
+			{
+				FBuildState& NewBuildState = PendingQueryState->QueryState[QueryIndex];
+				NewBuildState.Namespace = MoveTemp(Namespace);
+				NewBuildState.Platform = MoveTemp(Platform);
+				NewBuildState.Results = MoveTemp(Results);
+
+				if (--PendingQueryState->PendingQueries == 0)
+				{
+					// All queries complete
+					ExtractBuildsFromQueryResult(*PendingQueryState);
+
+					ExecuteOnGameThread(UE_SOURCE_LOCATION,
+						[this]
+						{
+							if (OnQueriesCompleteCallback)
+							{
+								OnQueriesCompleteCallback(PerPlatformBuilds);
+							}
+							ReadyToIssueBuildQuery = true;
+						});
+				}
+			});
+
+		++QueryIndex;
+	}
+}
+
+void FBuildListRetriever::ExtractBuildsFromQueryResult(FListBuildsState& ListBuildsState)
+{
+	for (FBuildState& BuildState : ListBuildsState.QueryState)
+	{
+		for (FBuildServiceInstance::FBuildRecord& BuildRecord : BuildState.Results)
+		{
+			TArray<int32>& Builds = PerPlatformBuilds.FindOrAdd(BuildState.Platform);
+			FString CommitIdentifier = BuildRecord.GetCommitIdentifier();
+			Builds.Add(FCString::Atoi(*CommitIdentifier));
+		}
+	}
+
+	// Sort builds in ascending order for each platform
+	for (auto& Pair : PerPlatformBuilds)
+	{
+		Pair.Value.Sort(TLess<int32>());
+	}
+}
+
+} // namespace UE::Zen::Build
diff --git a/Engine/Source/Developer/Zen/Private/BuildServerInterface.cpp b/Engine/Source/Developer/Zen/Private/BuildServerInterface.cpp
--- a/Engine/Source/Developer/Zen/Private/BuildServerInterface.cpp
+++ b/Engine/Source/Developer/Zen/Private/BuildServerInterface.cpp
@@ -485,6 +485,43 @@ FString FBuildServiceInstance::FBuildTransfer::GetDestination() const
 	return State->Destination;
 }
 
+FString FBuildServiceInstance::FBuildRecord::GetCommitIdentifier() const
+{
+	FString CommitIdentifier;
+	if (FCbFieldView ChangelistField = Metadata["changelist"]; ChangelistField.HasValue() && !ChangelistField.HasError())
+	{
+		if (ChangelistField.IsString())
+		{
+			CommitIdentifier = FUTF8ToTCHAR(ChangelistField.AsString());
+		}
+		else if (ChangelistField.IsInteger())
+		{
+			CommitIdentifier = *WriteToString<64>(ChangelistField.AsUInt64());
+		}
+		else if (ChangelistField.IsFloat())
+		{
+			CommitIdentifier = *WriteToString<64>((uint64)ChangelistField.AsDouble());
+		}
+	}
+	else if (FCbFieldView CommitField = Metadata["commit"]; CommitField.HasValue() && !CommitField.HasError())
+	{
+		if (CommitField.IsString())
+		{
+			CommitIdentifier = FUTF8ToTCHAR(CommitField.AsString());
+		}
+		else if (CommitField.IsInteger())
+		{
+			CommitIdentifier = *WriteToString<64>(CommitField.AsUInt64());
+		}
+		else if (CommitField.IsFloat())
+		{
+			CommitIdentifier = *WriteToString<64>((uint64)CommitField.AsDouble());
+		}
+	}
+
+	return CommitIdentifier;
+}
+
 FString
 FBuildServiceInstance::FBuildTransfer::GetDescription() const
 {
diff --git a/Engine/Source/Developer/Zen/Public/Experimental/BuildListRetriever.h b/Engine/Source/Developer/Zen/Public/Experimental/BuildListRetriever.h
new file mode 100644
--- /dev/null
+++ b/Engine/Source/Developer/Zen/Public/Experimental/BuildListRetriever.h
@@ -0,0 +1,44 @@
+// Copyright Epic Games, Inc. All Rights Reserved.
+
+#pragma once
+
+#include "Templates/SharedPointer.h"
+#include "Experimental/BuildServerInterface.h"
+#include "ZenBuildUtils.h"
+
+#define UE_API ZEN_API
+
+namespace UE::Zen::Build
+{
+	/**
+	 * Helper class to query a list of builds for a project of a specified build type in a P4 stream.
+	 * Call ConnectToBuildService() first to establish connection, before calling QueryBuilds() to retrieve the builds, build query is asynchronous.
+	 */
+	class FBuildListRetriever : public TSharedFromThis<FBuildListRetriever>
+	{
+	public:
+		UE_API FBuildListRetriever();
+		typedef TMap<FString, TArray<int32>> FPlatformBuildsMap;
+		typedef TUniqueFunction<void(const FPlatformBuildsMap&)> FOnQueriesComplete;
+		UE_API void ConnectToBuildService();
+		UE_API bool QueryBuilds(const FString& InProjectName, const FString& InBuildType, const FString& InStream, FOnQueriesComplete&& InOnQueriesCompleteCallback);
+
+	private:
+		void OnNamespacesAndBucketsRefreshed();
+		void IssueBuildQueries();
+		void ExtractBuildsFromQueryResult(FListBuildsState& ListBuildsState);
+
+		FString ProjectName;
+		FString BuildType;
+		FString Stream;
+
+		mutable TSharedPtr<FBuildServiceInstance> ServiceInstance;
+		FPlatformBuildsMap PerPlatformBuilds;
+		FOnQueriesComplete OnQueriesCompleteCallback;
+		bool ReadyToIssueBuildQuery = false; // Should only be accessed from the game thread
+		bool bConnected = false;
+	};
+
+} // namespace UE::Zen::Build
+
+#undef UE_API
diff --git a/Engine/Source/Developer/Zen/Public/Experimental/BuildServerInterface.h b/Engine/Source/Developer/Zen/Public/Experimental/BuildServerInterface.h
--- a/Engine/Source/Developer/Zen/Public/Experimental/BuildServerInterface.h
+++ b/Engine/Source/Developer/Zen/Public/Experimental/BuildServerInterface.h
@@ -69,6 +69,8 @@ public:
 	{
 		FCbObjectId BuildId;
 		FCbObject Metadata;
+
+		UE_API FString GetCommitIdentifier() const;
 	};
 
 	struct FBuildTransfer
diff --git a/Engine/Source/Editor/StorageServerWidgets/Internal/SBuildSelection.h b/Engine/Source/Editor/StorageServerWidgets/Internal/SBuildSelection.h
--- a/Engine/Source/Editor/StorageServerWidgets/Internal/SBuildSelection.h
+++ b/Engine/Source/Editor/StorageServerWidgets/Internal/SBuildSelection.h
@@ -16,13 +16,17 @@
 #include "Widgets/Views/SListView.h"
 #include "Widgets/Views/STableRow.h"
 #include "Widgets/Views/STableViewBase.h"
+#include "ZenBuildUtils.h"
 
 #define UE_API STORAGESERVERWIDGETS_API
 
 class FUProjectDictionary;
 class SMultiSelectComboBox;
 
-namespace UE::BuildSelection::Internal{ struct FListBuildsState; }
+namespace UE::Zen::Build
+{ 
+	struct FListBuildsState;
+}
 
 DECLARE_DELEGATE_ThreeParams(FOnBuildTransferStarted, UE::Zen::Build::FBuildServiceInstance::FBuildTransfer /*Transfer*/,
 	FStringView /*Name*/,
@@ -74,7 +78,7 @@ private:
 	static TPair<FRegexPattern, FKnownBuildType> MakeKnownBuildTypePattern(const FStringView InPattern, EBuildType Type, FText&& InText);
 	UE_API EBuildType GetSelectedBuildType() const;
 	UE_API void RebuildLists();
-	UE_API void RegenerateBuildGroups(UE::BuildSelection::Internal::FListBuildsState& ListBuildsState);
+	UE_API void RegenerateBuildGroups(UE::Zen::Build::FListBuildsState& ListBuildsState);
 	UE_API void ConditionalSortBuildGroups();
 	UE_API void RegenerateActivePlatformFilters();
 	UE_API void ValidateBuildGroupSelection();
diff --git a/Engine/Source/Editor/StorageServerWidgets/Private/SBuildSelection.cpp b/Engine/Source/Editor/StorageServerWidgets/Private/SBuildSelection.cpp
--- a/Engine/Source/Editor/StorageServerWidgets/Private/SBuildSelection.cpp
+++ b/Engine/Source/Editor/StorageServerWidgets/Private/SBuildSelection.cpp
@@ -43,26 +43,6 @@ namespace UE::BuildSelection::Internal
 		const FName ColCategory = TEXT("Category");
 		const FName ColCreated = TEXT("Created");
 	}
-
-	struct FNamespacePlatformBucketTuple
-	{
-		FString Namespace;
-		FString Platform;
-		FString Bucket;
-	};
-
-	struct FBuildState
-	{
-		FString Namespace;
-		FString Platform;
-		TArray<UE::Zen::Build::FBuildServiceInstance::FBuildRecord> Results;
-	};
-
-	struct FListBuildsState
-	{
-		std::atomic<uint32> PendingQueries;
-		TArray<FBuildState> QueryState;
-	};
 }
 
 TPair<FRegexPattern, SBuildSelection::FKnownBuildType> SBuildSelection::MakeKnownBuildTypePattern(const FStringView InPattern, EBuildType Type, FText&& InText)
@@ -452,7 +432,7 @@ void SBuildSelection::RebuildLists()
 	}
 }
 
-void SBuildSelection::RegenerateBuildGroups(UE::BuildSelection::Internal::FListBuildsState& ListBuildsState)
+void SBuildSelection::RegenerateBuildGroups(UE::Zen::Build::FListBuildsState& ListBuildsState)
 {
 	using namespace UE::Zen::Build;
 	using namespace UE::BuildSelection::Internal;
@@ -477,37 +457,8 @@ void SBuildSelection::RegenerateBuildGroups(UE::BuildSelection::Internal::FListB
 	{
 		for (FBuildServiceInstance::FBuildRecord& BuildRecord : BuildState.Results)
 		{
-			FString CommitIdentifier;
-			if (FCbFieldView ChangelistField = BuildRecord.Metadata["changelist"]; ChangelistField.HasValue() && !ChangelistField.HasError())
-			{
-				if (ChangelistField.IsString())
-				{
-					CommitIdentifier = FUTF8ToTCHAR(ChangelistField.AsString());
-				}
-				else if (ChangelistField.IsInteger())
-				{
-					CommitIdentifier = *WriteToString<64>(ChangelistField.AsUInt64());
-				}
-				else if (ChangelistField.IsFloat())
-				{
-					CommitIdentifier = *WriteToString<64>((uint64)ChangelistField.AsDouble());
-				}
-			}
-			else if (FCbFieldView CommitField = BuildRecord.Metadata["commit"]; CommitField.HasValue() && !CommitField.HasError())
-			{
-				if (CommitField.IsString())
-				{
-					CommitIdentifier = FUTF8ToTCHAR(CommitField.AsString());
-				}
-				else if (CommitField.IsInteger())
-				{
-					CommitIdentifier = *WriteToString<64>(CommitField.AsUInt64());
-				}
-				else if (CommitField.IsFloat())
-				{
-					CommitIdentifier = *WriteToString<64>((uint64)CommitField.AsDouble());
-				}
-			}
+			FString CommitIdentifier = BuildRecord.GetCommitIdentifier();
+
 			FString GroupKey = MakeGroupKey(BuildState.Namespace, CommitIdentifier, BuildRecord);
 			TSharedPtr<FBuildGroup>& BuildGroup = KeyedGroups.FindOrAdd(GroupKey);
 			if (!BuildGroup)
