Trivial change to variable name

diff --git a/Engine/Source/Developer/Zen/Private/BuildListRetriever.cpp b/Engine/Source/Developer/Zen/Private/BuildListRetriever.cpp
--- a/Engine/Source/Developer/Zen/Private/BuildListRetriever.cpp
+++ b/Engine/Source/Developer/Zen/Private/BuildListRetriever.cpp
@@ -134,9 +134,9 @@ void FBuildListRetriever::IssueBuildQueries()
 		// TODO: replace with passed in namespace
 		TArray<FString> NamespacesForBucket;
 		BucketsToNamespaces.MultiFind(Bucket, NamespacesForBucket);
-		for (FString& Namespace2 : NamespacesForBucket)
+		for (FString& Namespace : NamespacesForBucket)
 		{
-			NamespacePlatformBucketTuples.Emplace(MoveTemp(Namespace2), *Platform, Bucket);
+			NamespacePlatformBucketTuples.Emplace(MoveTemp(Namespace), *Platform, Bucket);
 		}
 	}
 
