BuildListRetriever - Fix potentially no build bucket available to query from causing future query unable to execute, and issue of not clearing build result in between queries

diff --git a/Engine/Source/Developer/Zen/Private/BuildListRetriever.cpp b/Engine/Source/Developer/Zen/Private/BuildListRetriever.cpp
--- a/Engine/Source/Developer/Zen/Private/BuildListRetriever.cpp
+++ b/Engine/Source/Developer/Zen/Private/BuildListRetriever.cpp
@@ -60,7 +60,6 @@ bool FBuildListRetriever::QueryBuilds(const FString& InProjectName, const FStrin
 void FBuildListRetriever::IssueBuildQueries()
 {
 	check(IsInGameThread());
-	ReadyToIssueBuildQuery = false;
 
 	TMultiMap<FString, FString> NamespacesAndBuckets = ServiceInstance->GetNamespacesAndBuckets();
 	TMultiMap<FString, FString> BucketsToNamespaces;
@@ -140,46 +139,52 @@ void FBuildListRetriever::IssueBuildQueries()
 		}
 	}
 
-	TSharedPtr<FListBuildsState> PendingQueryState = MakeShared<FListBuildsState>();
-	PendingQueryState->PendingQueries = NamespacePlatformBucketTuples.Num();
-	PendingQueryState->QueryState.SetNum(NamespacePlatformBucketTuples.Num());
-	uint32 QueryIndex = 0;
-
-	for (const FNamespacePlatformBucketTuple& NamespacePlatformBucket : NamespacePlatformBucketTuples)
+	if (NamespacePlatformBucketTuples.Num() > 0)
 	{
-		ServiceInstance->ListBuilds(NamespacePlatformBucket.Namespace, NamespacePlatformBucket.Bucket,
-			[this, QueryIndex, Namespace = NamespacePlatformBucket.Namespace, Platform = NamespacePlatformBucket.Platform,
-			PendingQueryState]
-			(TArray<FBuildServiceInstance::FBuildRecord>&& Results) mutable
-			{
-				FBuildState& NewBuildState = PendingQueryState->QueryState[QueryIndex];
-				NewBuildState.Namespace = MoveTemp(Namespace);
-				NewBuildState.Platform = MoveTemp(Platform);
-				NewBuildState.Results = MoveTemp(Results);
-
-				if (--PendingQueryState->PendingQueries == 0)
+		ReadyToIssueBuildQuery = false;
+		
+		TSharedPtr<FListBuildsState> PendingQueryState = MakeShared<FListBuildsState>();
+		PendingQueryState->PendingQueries = NamespacePlatformBucketTuples.Num();
+		PendingQueryState->QueryState.SetNum(NamespacePlatformBucketTuples.Num());
+		uint32 QueryIndex = 0;
+
+		for (const FNamespacePlatformBucketTuple& NamespacePlatformBucket : NamespacePlatformBucketTuples)
+		{
+			ServiceInstance->ListBuilds(NamespacePlatformBucket.Namespace, NamespacePlatformBucket.Bucket,
+				[this, QueryIndex, Namespace = NamespacePlatformBucket.Namespace, Platform = NamespacePlatformBucket.Platform,
+				PendingQueryState]
+				(TArray<FBuildServiceInstance::FBuildRecord>&& Results) mutable
 				{
-					// All queries complete
-					ExtractBuildsFromQueryResult(*PendingQueryState);
-
-					ExecuteOnGameThread(UE_SOURCE_LOCATION,
-						[this]
-						{
-							if (OnQueriesCompleteCallback)
+					FBuildState& NewBuildState = PendingQueryState->QueryState[QueryIndex];
+					NewBuildState.Namespace = MoveTemp(Namespace);
+					NewBuildState.Platform = MoveTemp(Platform);
+					NewBuildState.Results = MoveTemp(Results);
+
+					if (--PendingQueryState->PendingQueries == 0)
+					{
+						// All queries complete
+						ExtractBuildsFromQueryResult(*PendingQueryState);
+
+						ExecuteOnGameThread(UE_SOURCE_LOCATION,
+							[this]
 							{
-								OnQueriesCompleteCallback(PerPlatformBuilds);
-							}
-							ReadyToIssueBuildQuery = true;
-						});
-				}
-			});
-
-		++QueryIndex;
+								if (OnQueriesCompleteCallback)
+								{
+									OnQueriesCompleteCallback(PerPlatformBuilds);
+								}
+								ReadyToIssueBuildQuery = true;
+							});
+					}
+				});
+
+			++QueryIndex;
+		}
 	}
 }
 
 void FBuildListRetriever::ExtractBuildsFromQueryResult(FListBuildsState& ListBuildsState)
 {
+	PerPlatformBuilds.Reset();
 	for (FBuildState& BuildState : ListBuildsState.QueryState)
 	{
 		for (FBuildServiceInstance::FBuildRecord& BuildRecord : BuildState.Results)
