Only attempt system service uninstall if the default instance requests AutoLaunch and not SystemService.

[RN] none

diff --git a/Engine/Source/Developer/Zen/Private/ZenServerInterface.cpp b/Engine/Source/Developer/Zen/Private/ZenServerInterface.cpp
--- a/Engine/Source/Developer/Zen/Private/ZenServerInterface.cpp
+++ b/Engine/Source/Developer/Zen/Private/ZenServerInterface.cpp
@@ -2405,6 +2405,26 @@ FZenServiceInstance::GetAutoLaunchedPort()
 static bool UninstallSystemService(FString UtilityPath);
 static void PromptUserOfFailedSystemServiceUninstall();
 
+static bool IsDefaultInstanceSystemService()
+{
+	check(GConfig && GConfig->IsReadyForUse());
+	const TCHAR* ConfigSection = TEXT("Zen");
+	bool bSystemService = false;
+	GConfig->GetBool(ConfigSection, TEXT("SystemService"), bSystemService, GEngineIni);
+
+	return bSystemService;
+}
+
+static bool IsDefaultInstanceAutoLaunch()
+{
+	check(GConfig && GConfig->IsReadyForUse());
+	const TCHAR* ConfigSection = TEXT("Zen");
+	bool bAutoLaunch = false;
+	GConfig->GetBool(ConfigSection, TEXT("AutoLaunch"), bAutoLaunch, GEngineIni);
+
+	return bAutoLaunch;
+}
+
 void
 FZenServiceInstance::Initialize()
 {
@@ -2418,18 +2438,13 @@ FZenServiceInstance::Initialize()
 		Settings.bSystemService = false;
 	}
 
-	check(GConfig && GConfig->IsReadyForUse());
-	const TCHAR* ConfigSection = TEXT("Zen");
-	bool bSystemService = false;
-	GConfig->GetBool(ConfigSection, TEXT("SystemService"), bSystemService, GEngineIni);
-
 	// Creating a FSystemWideCriticalSection can sometimes fail if we haven't created
 	// ApplicationSettingsDir, as it attempts to take a file lock on a file in that
 	// directory. So we ensure it's created here before using FSystemWideCriticalSection.
 	IFileManager& FileManager = IFileManager::Get();
 	FileManager.MakeDirectory(FPlatformProcess::ApplicationSettingsDir());
 
-	if (!bSystemService)
+	if (!IsDefaultInstanceSystemService() && IsDefaultInstanceAutoLaunch())
 	{
 		FSystemWideCriticalSection InstallCriticalSection(TEXT("ZenServerInstall"), FTimespan::FromSeconds(FApp::IsUnattended() ? 30.0 : 10.0));
 		if (!InstallCriticalSection.IsValid())
