Allow zenserver installation as daemon on Linux.

[RN] minor

diff --git a/Engine/Source/Developer/Zen/Private/ZenServerInterface.cpp b/Engine/Source/Developer/Zen/Private/ZenServerInterface.cpp
--- a/Engine/Source/Developer/Zen/Private/ZenServerInterface.cpp
+++ b/Engine/Source/Developer/Zen/Private/ZenServerInterface.cpp
@@ -359,8 +359,16 @@ FString WriteMiniDump(HANDLE ProcessHandle, const FString& ExeName, const FStrin
 
 static bool RunZenUtility(const FString& UtilityPath, const FString& CommandLine, FString& Output, int& ReturnCode)
 {
-	FString AbsoluteUtilityPath = FPaths::ConvertRelativePathToFull(UtilityPath);
-	FPaths::MakePlatformFilename(AbsoluteUtilityPath);
+	FString AbsoluteUtilityPath = UtilityPath;
+
+	// If we're running zen on Linux/Mac using 'sudo' or 'pkexec' we don't want to make those paths relative.
+	// In that case UtilityPath is 'pkexec'/'sudo' and the full path to the zen binary is already part of
+	// CommandLine.
+	if (AbsoluteUtilityPath.StartsWith(TEXT("zen")))
+	{
+		AbsoluteUtilityPath = FPaths::ConvertRelativePathToFull(UtilityPath);
+		FPaths::MakePlatformFilename(AbsoluteUtilityPath);
+	}
 
 	bool bLaunchDetached = true;
 	bool bLaunchHidden = true;
@@ -396,7 +404,13 @@ static bool RunZenUtility(const FString& UtilityPath, const FString& CommandLine
 
 	bool DumpWritten = false;
 
-	const double MaxWaitDuration = FApp::IsUnattended() ? 25.0 : 8.0;
+	double MaxWaitDuration = FApp::IsUnattended() ? 25.0 : 8.0;
+
+#if PLATFORM_LINUX
+	// Linux can require interactive password entry when installing system services, so increase the timeout
+	// if we're running this under 'sudo' or 'pkexec' rather than directly shelling out to zen
+	MaxWaitDuration = (UtilityPath.StartsWith("zen") ? MaxWaitDuration : 30.0);
+#endif
 
 	const uint64 StartTime = FPlatformTime::Cycles64();
 	while (FPlatformProcess::IsProcRunning(ZenUtilProc))
@@ -468,6 +482,23 @@ static bool RunZenUtility(const FString& UtilityPath, const FString& CommandLine
 	}
 }
 
+static bool RunZenUtilityElevated(const FString& UtilityPath, const FString& CommandLine, FString& Output, int& ReturnCode)
+{
+	FString ElevatedCommandLine;
+	FString ElevatedUtilityPath;
+
+#if PLATFORM_WINDOWS
+	ElevatedUtilityPath = UtilityPath;
+	ElevatedCommandLine = TEXT("--allow-elevation ") + CommandLine;
+#endif
+#if PLATFORM_LINUX
+	ElevatedCommandLine = UtilityPath + TEXT(" ") + CommandLine;
+	ElevatedUtilityPath = TEXT("pkexec");
+#endif
+
+	return RunZenUtility(ElevatedUtilityPath, ElevatedCommandLine, Output, ReturnCode);
+}
+
 static bool GetZenVersion(const FString& UtilityPath, const FString& ServicePath, FZenVersion& OutVersion)
 {
 	FString TempOutputFilePath = FPaths::CreateTempFilename(FPlatformProcess::UserTempDir(), TEXT("zenversion"));
@@ -2836,7 +2867,7 @@ bool UninstallSystemService(FString UtilityPath)
 			return false;
 		}
 
-		if (!RunZenUtility(UtilityPath, TEXT("service stop --allow-elevation"), Output, ReturnCode))
+		if (!RunZenUtilityElevated(UtilityPath, TEXT("service stop"), Output, ReturnCode))
 		{
 			PromptUserOfFailedSystemServiceUninstall();
 			return false;
@@ -2868,13 +2899,13 @@ bool UninstallSystemService(FString UtilityPath)
 		return true;
 	}
 
-	if (!RunZenUtility(UtilityPath, TEXT("service uninstall --allow-elevation"), Output, ReturnCode))
+	if (!RunZenUtilityElevated(UtilityPath, TEXT("service uninstall"), Output, ReturnCode))
 	{
 		PromptUserOfFailedSystemServiceUninstall();
 		return false;
 	}
 
-	if (ReturnCode == -1)
+	if (ReturnCode != 0)
 	{
 		PromptUserNeedsElevatedPermissions();
 		return false;
@@ -2911,28 +2942,19 @@ bool FZenServiceInstance::ConditionalUpdateSystemServiceInstall(const FSystemSer
 	FString DesiredCommandLine = DetermineCmdLineWithoutTransientComponents(InSettings.DataPath, InSettings.ExtraArgs, Plugins, InSettings.bSendUnattendedBugReports, InSettings.bAllowRemoteNetworkService, InSettings.DesiredPort);
 
 	FZenVersion Version;
-	bool bIsServiceRunning = DetermineSystemServiceInfo(SystemUtilityInstallPath, RunningExecutable, RunningCommandLine, RunningPort);
+	bool bIsServiceRunning = DetermineSystemServiceInfo(InTreeUtilityPath, RunningExecutable, RunningCommandLine, RunningPort);
 	bool bServiceNeedsUpdate = !bIsServiceRunning
 		|| RunningCommandLine != DesiredCommandLine
 		|| IsInstallVersionOutOfDate(InTreeUtilityPath, SystemUtilityInstallPath, InTreeServicePath, SystemServiceInstallPath, InTreeVersionCache, SystemVersionCache, Version);
 
 	if (bServiceNeedsUpdate)
 	{
-		FString InstallDirectory = GetSystemInstallPath(InSettings);
-		if (!IsDirectoryWriteable(InstallDirectory))
-		{
-			PromptUserUnableToUpdateSystemServiceBinaries(InstallDirectory);
-			return false;
-		}
-
-		if (!UninstallSystemService(SystemUtilityInstallPath))
-		{
-			return false;
-		}
-
 		const ZenServerState State(/*ReadOnly*/true);
 		const ZenServerState::ZenServerEntry* RunningEntry = State.LookupByDesiredListenPort(InSettings.DesiredPort);
-		if (RunningEntry != nullptr)
+		bool bIsServiceRunningOnPort = (bIsServiceRunning && RunningEntry->EffectiveListenPort == RunningPort);
+		// If an existing system service install is running on this port, 'service install' will handle
+		// shutdown for us. Otherwise, we need to request that the non-system-service zen shut down manually.
+		if (RunningEntry != nullptr && !bIsServiceRunningOnPort)
 		{
 			if (!ShutdownZenServerProcess(RunningEntry->Pid))
 			{
@@ -2941,34 +2963,20 @@ bool FZenServiceInstance::ConditionalUpdateSystemServiceInstall(const FSystemSer
 			}
 		}
 
-		if (!UpdateServiceBinaries(InTreeUtilityPath, SystemUtilityInstallPath, InTreeServicePath, SystemServiceInstallPath, InTreeVersionCache, SystemVersionCache))
-		{
-			PromptUserUnableToUpdateSystemServiceBinaries(InstallDirectory);
-			return false;
-		}
-
 		FString ServiceInstallCommandLine;
-		ServiceInstallCommandLine.Appendf(TEXT("service install --allow-elevation -- %s"), *DesiredCommandLine);
-
-		if (!RunZenUtility(SystemUtilityInstallPath, ServiceInstallCommandLine, Output, ReturnCode))
-		{
-			PromptUserOfFailedSystemServiceInstall();
-			return false;
-		}
-
-		if (ReturnCode == -1)
-		{
-			PromptUserNeedsElevatedPermissions();
-			return false;
-		}
+		ServiceInstallCommandLine.Appendf(TEXT("service install --full --install-path \"%s\""), *GetSystemInstallPath(InSettings));
+#if PLATFORM_LINUX
+		ServiceInstallCommandLine.Appendf(TEXT(" --user %s"), FPlatformProcess::UserName(false));
+#endif
+		ServiceInstallCommandLine.Appendf(TEXT(" -- %s"), *DesiredCommandLine);
 
-		if (!RunZenUtility(SystemUtilityInstallPath, TEXT("service start --allow-elevation"), Output, ReturnCode))
+		if (!RunZenUtilityElevated(InTreeUtilityPath, ServiceInstallCommandLine, Output, ReturnCode))
 		{
 			PromptUserOfFailedSystemServiceInstall();
 			return false;
 		}
 
-		if (ReturnCode == -1)
+		if (ReturnCode != 0)
 		{
 			PromptUserNeedsElevatedPermissions();
 			return false;
diff --git a/Engine/Source/Runtime/Core/Private/Unix/UnixPlatformProcess.cpp b/Engine/Source/Runtime/Core/Private/Unix/UnixPlatformProcess.cpp
--- a/Engine/Source/Runtime/Core/Private/Unix/UnixPlatformProcess.cpp
+++ b/Engine/Source/Runtime/Core/Private/Unix/UnixPlatformProcess.cpp
@@ -202,9 +202,11 @@ const TCHAR* FUnixPlatformProcess::ComputerName()
 const TCHAR* FUnixPlatformProcess::UserName(bool bOnlyAlphaNumeric)
 {
 	static TCHAR Name[PlatformProcessLimits::MaxUserName] = { 0 };
+	static TCHAR AlphaNumericName[PlatformProcessLimits::MaxUserName] = { 0 };
 	static bool bHaveResult = false;
+	static bool bHaveAlphaNumericResult = false;
 
-	if (!bHaveResult)
+	if ((!bHaveResult && !bOnlyAlphaNumeric) || (!bHaveAlphaNumericResult && bOnlyAlphaNumeric))
 	{
 		struct passwd * UserInfo = getpwuid(geteuid());
 		if (nullptr != UserInfo && nullptr != UserInfo->pw_name)
@@ -213,8 +215,8 @@ const TCHAR* FUnixPlatformProcess::UserName(bool bOnlyAlphaNumeric)
 			if (bOnlyAlphaNumeric)
 			{
 				const TCHAR *Src = *TempName;
-				TCHAR * Dst = Name;
-				for (; *Src != 0 && (Dst - Name) < UE_ARRAY_COUNT(Name) - 1; ++Src)
+				TCHAR * Dst = AlphaNumericName;
+				for (; *Src != 0 && (Dst - AlphaNumericName) < UE_ARRAY_COUNT(AlphaNumericName) - 1; ++Src)
 				{
 					if (FChar::IsAlnum(*Src))
 					{
@@ -222,20 +224,24 @@ const TCHAR* FUnixPlatformProcess::UserName(bool bOnlyAlphaNumeric)
 					}
 				}
 				*Dst++ = 0;
+				bHaveAlphaNumericResult = true;
 			}
 			else
 			{
 				FCString::Strncpy(Name, *TempName, UE_ARRAY_COUNT(Name) - 1);
+				bHaveResult = true;
 			}
 		}
 		else
 		{
 			FCString::Sprintf(Name, TEXT("euid%d"), geteuid());
+			FCString::Sprintf(AlphaNumericName, TEXT("euid%d"), geteuid());
+			bHaveResult = true;
+			bHaveAlphaNumericResult = true;
 		}
-		bHaveResult = true;
 	}
 
-	return Name;
+	return bOnlyAlphaNumeric ? AlphaNumericName : Name;
 }
 
 const TCHAR* FUnixPlatformProcess::UserTempDir()
