Add missing null check in ZenServerInterface.

[RN] none

diff --git a/Engine/Source/Developer/Zen/Private/ZenServerInterface.cpp b/Engine/Source/Developer/Zen/Private/ZenServerInterface.cpp
--- a/Engine/Source/Developer/Zen/Private/ZenServerInterface.cpp
+++ b/Engine/Source/Developer/Zen/Private/ZenServerInterface.cpp
@@ -2951,7 +2951,7 @@ bool FZenServiceInstance::ConditionalUpdateSystemServiceInstall(const FSystemSer
 	{
 		const ZenServerState State(/*ReadOnly*/true);
 		const ZenServerState::ZenServerEntry* RunningEntry = State.LookupByDesiredListenPort(InSettings.DesiredPort);
-		bool bIsServiceRunningOnPort = (bIsServiceRunning && RunningEntry->EffectiveListenPort == RunningPort);
+		bool bIsServiceRunningOnPort = (bIsServiceRunning && RunningEntry && RunningEntry->EffectiveListenPort == RunningPort);
 		// If an existing system service install is running on this port, 'service install' will handle
 		// shutdown for us. Otherwise, we need to request that the non-system-service zen shut down manually.
 		if (RunningEntry != nullptr && !bIsServiceRunningOnPort)
