Add a missing logical negation that was causing zenserver to be started with the no-sentry commandline argument inappropriately.

diff --git a/Engine/Source/Developer/Zen/Private/ZenServerInterface.cpp b/Engine/Source/Developer/Zen/Private/ZenServerInterface.cpp
--- a/Engine/Source/Developer/Zen/Private/ZenServerInterface.cpp
+++ b/Engine/Source/Developer/Zen/Private/ZenServerInterface.cpp
@@ -1775,7 +1775,7 @@ DetermineCmdLineWithoutTransientComponents(const FString& DataPath, const FStrin
 		}
 	}
 
-	if (bSendUnattendedBugReports)
+	if (!bSendUnattendedBugReports)
 	{
 		Parms.Append(TEXT(" --no-sentry"));
 	}
