Update zenserver to v5.7.6 [9a1ebb9] https://github.com/EpicGames/zen/releases/tag/v5.7.6

v5.7.6
- Improvement: Oidc token executable is now launched hidden when launch from inside zenserver (oplog import/export)
- Improvement: Helper scripts `get_ue_toolchain.h` and `ue_build.sh` can now be used without passing an argument to them. If you don't specify a toolchain directory then the default of `.tmp-ue-toolchain` is used
- Bugfix: Fixed bug in block store GC where a chunk could be lost due to writing blocks larger than max block size UE-330662

Fixed problems with detecting zenserver running as a service
Remove no longer used ActiveSponsors member of ZenServerState since it is no longer is part of the zenserver structure
Fixed parsing of 'zen service status' command

#rnx

diff --git a/Engine/Source/Developer/Zen/Private/ZenServerInterface.cpp b/Engine/Source/Developer/Zen/Private/ZenServerInterface.cpp
--- a/Engine/Source/Developer/Zen/Private/ZenServerInterface.cpp
+++ b/Engine/Source/Developer/Zen/Private/ZenServerInterface.cpp
@@ -2767,13 +2767,14 @@ static bool DetermineSystemServiceInfo(const FString UtilityPath, FString& OutEx
 	}
 
 	FString CommandLineOptionsSubstring = TEXT("CommandLineOptions:");
-	int32 ArgsIndex = Output.Find(CommandLineOptionsSubstring) + CommandLineOptionsSubstring.Len();
+	int32 ArgsIndex = Output.Find(CommandLineOptionsSubstring);
 	if (ArgsIndex == INDEX_NONE)
 	{
 		UE_LOG(LogZenServiceInstance, Warning, TEXT("Failed to determine command line options for running Unreal Zen Storage Server instance"));
 		return false;
 	}
 
+	ArgsIndex += CommandLineOptionsSubstring.Len();
 	int32 ArgsEnd = Output.Find("\n", ESearchCase::IgnoreCase, ESearchDir::FromStart, ArgsIndex);
 	if (ArgsEnd == INDEX_NONE)
 	{
@@ -2784,12 +2785,13 @@ static bool DetermineSystemServiceInfo(const FString UtilityPath, FString& OutEx
 	OutCommandLine = Output.Mid(ArgsIndex, ArgsEnd - ArgsIndex).TrimStartAndEnd();
 
 	FString PortSubstring = TEXT("--port ");
-	int32 PortIndex = OutCommandLine.Find(PortSubstring) + PortSubstring.Len();
+	int32 PortIndex = OutCommandLine.Find(PortSubstring);
 	if (PortIndex == INDEX_NONE)
 	{
 		UE_LOG(LogZenServiceInstance, Warning, TEXT("Failed to determine port for running Unreal Zen Storage Server instance"));
 		return false;
 	}
+	PortIndex += PortSubstring.Len();
 
 	int32 PortEnd = OutCommandLine.Find(" ", ESearchCase::IgnoreCase, ESearchDir::FromStart, PortIndex);
 	if (PortEnd == INDEX_NONE)
@@ -2805,12 +2807,13 @@ static bool DetermineSystemServiceInfo(const FString UtilityPath, FString& OutEx
 	}
 
 	FString ExecutableSubstring = TEXT("Executable:");
-	int32 ExecutableIndex = Output.Find(ExecutableSubstring) + ExecutableSubstring.Len();
+	int32 ExecutableIndex = Output.Find(ExecutableSubstring);
 	if (ExecutableIndex == INDEX_NONE)
 	{
 		UE_LOG(LogZenServiceInstance, Warning, TEXT("Failed to determine executable for running Unreal Zen Storage Server instance"));
 		return false;
 	}
+	ExecutableIndex += ExecutableSubstring.Len();
 
 	int32 ExecutableEnd = Output.Find("\n", ESearchCase::IgnoreCase, ESearchDir::FromStart, ExecutableIndex);
 	if (ExecutableEnd == INDEX_NONE)
@@ -2859,11 +2862,7 @@ bool UninstallSystemService(FString UtilityPath)
 
 		if (!Entry)
 		{
-			UE_LOG(LogZenServiceInstance, Warning, TEXT("Failed to look up running service info to determine active users of service. Assuming no active users."));
-		}
-		else if (Entry->ActiveSponsors.load())
-		{
-			PromptUserToStopRunningServerInstanceForUpdate(FPaths::GetPath(UtilityPath));
+			UE_LOG(LogZenServiceInstance, Warning, TEXT("Failed to look up running service info to determine active users of service, assuming it is not running"));
 			return false;
 		}
 
diff --git a/Engine/Source/Developer/Zen/Private/ZenServerState.h b/Engine/Source/Developer/Zen/Private/ZenServerState.h
--- a/Engine/Source/Developer/Zen/Private/ZenServerState.h
+++ b/Engine/Source/Developer/Zen/Private/ZenServerState.h
@@ -23,8 +23,7 @@ public:
 		uint8				SessionId[12];
 		std::atomic<uint32> SponsorPids[8];
 		std::atomic<uint16> EffectiveListenPort;
-		std::atomic<bool>   ActiveSponsors;
-		uint8				Padding[9];
+		uint8				Padding[10];
 
 		enum class FlagsEnum : uint16
 		{
