Warn user if a network data path is used for zenserver

#rnx

diff --git a/Engine/Source/Developer/Zen/Private/ZenServerInterface.cpp b/Engine/Source/Developer/Zen/Private/ZenServerInterface.cpp
--- a/Engine/Source/Developer/Zen/Private/ZenServerInterface.cpp
+++ b/Engine/Source/Developer/Zen/Private/ZenServerInterface.cpp
@@ -1004,21 +1004,40 @@ PromptUserAboutInvalidValidDataPathConfiguration(const FString& UsedDataPath)
 
 #if PLATFORM_WINDOWS
 static void
-PromptUserIsUsingGoogleDriveAsDataPath()
+PromptUserIsUsingGoogleDriveAsDataPath(const FString& BadDataPath)
 {
 #if !IS_PROGRAM
 	if (!FApp::IsUnattended() && !IsRunningCommandlet() && !GIsRunningUnattendedScript)
 	{
 		FString LogDirPath = FPaths::ConvertRelativePathToFull(FPaths::ProjectLogDir());
 		FText ZenInvalidDataPathPromptTitle = NSLOCTEXT("Zen", "Zen_GoogleDriveDataPathPromptTitle", "Using Google Drive as a data path");
-		FText ZenInvalidDataPathPromptText = FText::Format(NSLOCTEXT("Zen", "Zen_GoogleDriveDataPathPromptText", "Unreal Zen Storage Server is configured to use Google Drive as a data path, this is highly inadvisable.\nPlease use a data path on a local physical drive.\nCheck the log in '{0}' for details.\nUpdate your configuration and restart."), FText::FromString(LogDirPath));
+		FText ZenInvalidDataPathPromptText = FText::Format(NSLOCTEXT("Zen", "Zen_GoogleDriveDataPathPromptText", "Unreal Zen Storage Server is configured to use Google Drive '{0}' as a data path, this is highly inadvisable.\nPlease use a data path on a local physical drive.\nCheck the log in '{1}' for details.\nUpdate your configuration and restart."), FText::FromString(BadDataPath), FText::FromString(LogDirPath));
 		FPlatformMisc::MessageBoxExt(EAppMsgType::Ok, *ZenInvalidDataPathPromptText.ToString(), *ZenInvalidDataPathPromptTitle.ToString());
 	}
 	else
 #endif
 	{
 		// Just log as there is no one to show a message
-		UE_LOG(LogZenServiceInstance, Warning, TEXT("Unreal Zen Storage Server is configured to use Google Drive as a data path, this is highly inadvisable. Please use a path on a local physical drive."));
+		UE_LOG(LogZenServiceInstance, Warning, TEXT("Unreal Zen Storage Server is configured to use Google Drive as a data path '%s', this is highly inadvisable. Please use a path on a local physical drive."), *BadDataPath);
+	}
+}
+
+static void
+PromptUserIsUsingNetworkDriveAsDataPath(const FString& BadDataPath)
+{
+#if !IS_PROGRAM
+	if (!FApp::IsUnattended() && !IsRunningCommandlet() && !GIsRunningUnattendedScript)
+	{
+		FString LogDirPath = FPaths::ConvertRelativePathToFull(FPaths::ProjectLogDir());
+		FText ZenInvalidDataPathPromptTitle = NSLOCTEXT("Zen", "Zen_NetworkDriveDataPathPromptTitle", "Using Network Drive as a data path");
+		FText ZenInvalidDataPathPromptText = FText::Format(NSLOCTEXT("Zen", "Zen_NetworkDriveDataPathPromptText", "Unreal Zen Storage Server is configured to a network path '{0}' as a data path, this is highly inadvisable.\nPlease use a data path on a local physical drive.\nCheck the log in '{1}' for details.\nUpdate your configuration and restart."), FText::FromString(BadDataPath), FText::FromString(LogDirPath));
+		FPlatformMisc::MessageBoxExt(EAppMsgType::Ok, *ZenInvalidDataPathPromptText.ToString(), *ZenInvalidDataPathPromptTitle.ToString());
+	}
+	else
+#endif
+	{
+		// Just log as there is no one to show a message
+		UE_LOG(LogZenServiceInstance, Warning, TEXT("Unreal Zen Storage Server is configured to use Network Drive as a data path '%s', this is highly inadvisable. Please use a path on a local physical drive."), *BadDataPath);
 	}
 }
 #endif // PLATFORM_WINDOWS
@@ -1249,27 +1268,40 @@ FServiceSettings::ReadFromConfig()
 				if (AutoLaunchSettings.DataPath.FindChar(':', DriveEnd))
 				{
 					FString DrivePath = AutoLaunchSettings.DataPath.Left(DriveEnd + 1);
-					TCHAR VolumeName[128];
-
-					BOOL OK = GetVolumeInformation(
-						*DrivePath,
-						VolumeName,
-						127,
-						NULL,
-						NULL,
-						NULL,
-						NULL,
-						NULL);
-
-					if (OK)
+
+					UINT DriveType = GetDriveType(*DrivePath);
+					if (DriveType == DRIVE_REMOTE)
 					{
-						VolumeName[127] = 0;
-						if (FString(VolumeName) == TEXT("Google Drive"))
+						PromptUserIsUsingNetworkDriveAsDataPath(AutoLaunchSettings.DataPath);
+					}
+					else
+					{
+						TCHAR VolumeName[128];
+
+						BOOL OK = GetVolumeInformation(
+							*DrivePath,
+							VolumeName,
+							127,
+							NULL,
+							NULL,
+							NULL,
+							NULL,
+							NULL);
+
+						if (OK)
 						{
-							PromptUserIsUsingGoogleDriveAsDataPath();
+							VolumeName[127] = 0;
+							if (FString(VolumeName) == TEXT("Google Drive"))
+							{
+								PromptUserIsUsingGoogleDriveAsDataPath(AutoLaunchSettings.DataPath);
+							}
 						}
 					}
 				}
+				else if (AutoLaunchSettings.DataPath.StartsWith("\\\\") || AutoLaunchSettings.DataPath.StartsWith("//"))
+				{
+					PromptUserIsUsingNetworkDriveAsDataPath(AutoLaunchSettings.DataPath);
+				}
 			}
 #endif // PLATFORM_WINDOWS
 
