Pass correct root drive string to GetDriveType and GetVolumeInformation.

#rnx
#FailedValidations Preflight

diff --git a/Engine/Source/Developer/Zen/Private/ZenServerInterface.cpp b/Engine/Source/Developer/Zen/Private/ZenServerInterface.cpp
--- a/Engine/Source/Developer/Zen/Private/ZenServerInterface.cpp
+++ b/Engine/Source/Developer/Zen/Private/ZenServerInterface.cpp
@@ -1267,7 +1267,7 @@ FServiceSettings::ReadFromConfig()
 				int32 DriveEnd = 0;
 				if (AutoLaunchSettings.DataPath.FindChar(':', DriveEnd))
 				{
-					FString DrivePath = AutoLaunchSettings.DataPath.Left(DriveEnd + 1);
+					FString DrivePath = AutoLaunchSettings.DataPath.Left(DriveEnd + 2);
 
 					UINT DriveType = GetDriveType(*DrivePath);
 					if (DriveType == DRIVE_REMOTE)
