BuildStorageTool - skip empty lines in log for UI display

#rnx

diff --git a/Engine/Source/Developer/Zen/Private/BuildServerInterface.cpp b/Engine/Source/Developer/Zen/Private/BuildServerInterface.cpp
--- a/Engine/Source/Developer/Zen/Private/BuildServerInterface.cpp
+++ b/Engine/Source/Developer/Zen/Private/BuildServerInterface.cpp
@@ -1041,51 +1041,54 @@ FBuildServiceInstance::BuildTransferThreadLoop()
 		MonitoredUtilityProcess.OnOutput().BindSPLambda(State, [State](FString Output)
 		{
 			FReadScopeLock _(State->Lock);
-			if (Output.StartsWith(TEXT("@progress push")))
-			{
-				Output = Output.RightChop(15);
-				Private::FProgressState NewState;
-				NewState.Label = Output;
-				State->ProgressState.Push(MoveTemp(NewState));
-				return;
-			}
-			else if (Output.StartsWith(TEXT("@progress pop")))
+			if (State->OutputLogFile)
 			{
-				State->ProgressState.Pop();
-				return;
+				State->OutputLogFile->Logf(TEXT("%s"), *Output);	// Log should be unaltered output
+				State->OutputLogFile->Flush();
 			}
-			else if (Output.StartsWith(TEXT("@progress")))
+			if (!Output.IsEmpty())	// Skip empty lines
 			{
-				if (State->ProgressState.IsEmpty())
+				if (Output.StartsWith(TEXT("@progress push")))
 				{
+					Output = Output.RightChop(15);
 					Private::FProgressState NewState;
+					NewState.Label = Output;
 					State->ProgressState.Push(MoveTemp(NewState));
+					return;
+				}
+				else if (Output.StartsWith(TEXT("@progress pop")))
+				{
+					State->ProgressState.Pop();
+					return;
 				}
-				Private::FProgressState& CurrentState = State->ProgressState.Top();
-				Output = Output.RightChop(10);
-				if (Output.EndsWith(TEXT("%")))
+				else if (Output.StartsWith(TEXT("@progress")))
 				{
-					FString PossiblePercentText = Output.LeftChop(1);
-					if (!PossiblePercentText.IsEmpty())
+					if (State->ProgressState.IsEmpty())
+					{
+						Private::FProgressState NewState;
+						State->ProgressState.Push(MoveTemp(NewState));
+					}
+					Private::FProgressState& CurrentState = State->ProgressState.Top();
+					Output = Output.RightChop(10);
+					if (Output.EndsWith(TEXT("%")))
 					{
-						float Percent = FCString::Atof(*PossiblePercentText);
-						if ((Percent != 0.0f) || (PossiblePercentText.Len() == 1 && PossiblePercentText[0] == TCHAR('0')))
+						FString PossiblePercentText = Output.LeftChop(1);
+						if (!PossiblePercentText.IsEmpty())
 						{
-							CurrentState.Percent = Percent;
-							return; // Do not print the percent progress alone
+							float Percent = FCString::Atof(*PossiblePercentText);
+							if ((Percent != 0.0f) || (PossiblePercentText.Len() == 1 && PossiblePercentText[0] == TCHAR('0')))
+							{
+								CurrentState.Percent = Percent;
+								return; // Do not print the percent progress alone
+							}
 						}
 					}
+					CurrentState.Detail = Output;
+					Output = FString::Printf(TEXT("[%s %d%%] %s"), *CurrentState.Label, (int32)CurrentState.Percent, *Output);
+					// Deliberate fall-through so that the text is printed;
 				}
-				CurrentState.Detail = Output;
-				Output = FString::Printf(TEXT("[%s %d%%] %s"), *CurrentState.Label, (int32)CurrentState.Percent, *Output);
-				// Deliberate fall-through so that the text is printed;
-			}
-			if (State->OutputLogFile)
-			{
-				State->OutputLogFile->Logf(TEXT("%s"), *Output);
-				State->OutputLogFile->Flush();
+				State->RecentOutputLines[State->NextOutputLineIndex++] = MoveTemp(Output);
 			}
-			State->RecentOutputLines[State->NextOutputLineIndex++] = MoveTemp(Output);
 		});
 
 		{
