Change build downloading to use oidctoken.exe as an argument to the download utility.

diff --git a/Engine/Source/Developer/Zen/Private/BuildServerInterface.cpp b/Engine/Source/Developer/Zen/Private/BuildServerInterface.cpp
--- a/Engine/Source/Developer/Zen/Private/BuildServerInterface.cpp
+++ b/Engine/Source/Developer/Zen/Private/BuildServerInterface.cpp
@@ -1029,10 +1029,10 @@ FBuildServiceInstance::BuildTransferThreadLoop()
 						*DestinationProjectId, *DestinationOplogId));
 				}
 
-				CommandLineArgs = FString::Printf(TEXT("oplog-import %s %s --clean --builds %s --namespace %s --bucket %s --builds-id %s --access-token \"%s\""),
+				CommandLineArgs = FString::Printf(TEXT("oplog-import %s %s --clean --builds %s --namespace %s --bucket %s --builds-id %s --oidctoken-exe-path \"%s\""),
 					*DestinationProjectId, *DestinationOplogId, *WriteToString<64>(EffectiveDomain),
 					*State->Namespace, *State->Bucket,
-					*WriteToString<64>(State->BuildId), *GetAccessToken());
+					*WriteToString<64>(State->BuildId), *OidcTokenExecutableFilename);
 			}
 			break;
 		}
