Fix typo from CL 46599368

#rnx

diff --git a/Engine/Source/Developer/Zen/Private/ZenServerInterface.cpp b/Engine/Source/Developer/Zen/Private/ZenServerInterface.cpp
--- a/Engine/Source/Developer/Zen/Private/ZenServerInterface.cpp
+++ b/Engine/Source/Developer/Zen/Private/ZenServerInterface.cpp
@@ -2856,7 +2856,7 @@ static bool DetermineSystemServiceInfo(const FString UtilityPath, FString& OutEx
 		return false;
 	}
 
-	ArgsIndex = +CommandLineOptionsSubstring.Len();
+	ArgsIndex += CommandLineOptionsSubstring.Len();
 	int32 ArgsEnd = Output.Find("\n", ESearchCase::IgnoreCase, ESearchDir::FromStart, ArgsIndex);
 	if (ArgsEnd == INDEX_NONE)
 	{
