BuildStorageTool: Usability fixes
- Fixing ticker delta time bug that caused auth token refreshes not happen on schedule
- Improve the selection of the default destination directory for file-based downloads to prefer user-selected or fallback workspaces over the tool workspace
- Adding transfer duration tracking to build activities
- Changed label for oplogs from Cooked Data to Cook Snapshots for more consistent terminology

diff --git a/Engine/Source/Developer/Zen/Private/BuildServerInterface.cpp b/Engine/Source/Developer/Zen/Private/BuildServerInterface.cpp
--- a/Engine/Source/Developer/Zen/Private/BuildServerInterface.cpp
+++ b/Engine/Source/Developer/Zen/Private/BuildServerInterface.cpp
@@ -76,6 +76,8 @@ struct FBuildTransferState : public TSharedFromThis<FBuildTransferState, ESPMode
 	TUniquePtr<FArchive> OutputLogFile;
 	UE::Zen::Build::FBuildServiceInstance::EBuildTransferStatus Status = UE::Zen::Build::FBuildServiceInstance::EBuildTransferStatus::Queued;
 	int ReturnCode;
+	FMonotonicTimePoint TransferStartTime;
+	FMonotonicTimePoint TransferEndTime;
 
 	FString GetRecentOutput() const;
 	FString GetLogFilename() const;
@@ -614,6 +616,28 @@ FBuildServiceInstance::FBuildTransfer::GetLogFilename() const
 	return State->GetLogFilename();
 }
 
+FMonotonicTimeSpan
+FBuildServiceInstance::FBuildTransfer::GetTransferDuration() const
+{
+	if (!State)
+	{
+		return FMonotonicTimeSpan::Zero();
+	}
+
+	FReadScopeLock _(State->Lock);
+	if (State->Status == EBuildTransferStatus::Canceled ||
+		State->Status == EBuildTransferStatus::Failed ||
+		State->Status == EBuildTransferStatus::Succeeded)
+	{
+		return State->TransferEndTime - State->TransferStartTime;
+	}
+	else if (State->Status == EBuildTransferStatus::Active)
+	{
+		return FMonotonicTimePoint::Now() - State->TransferStartTime;
+	}
+	return FMonotonicTimeSpan::Zero();
+}
+
 bool
 FBuildServiceInstance::FBuildTransfer::GetCurrentProgress(FString& OutLabel, FString& OutDetail, float& OutPercent) const
 {
@@ -681,6 +705,7 @@ FBuildServiceInstance::FBuildTransfer::RequestCancel()
 	}
 	State->bCancelRequested.store(true);
 	FWriteScopeLock _(State->Lock);
+	State->TransferEndTime = FMonotonicTimePoint::Now();
 	State->Status = EBuildTransferStatus::Canceled;
 }
 
@@ -723,6 +748,8 @@ FBuildServiceInstance::AcquireAccessToken(EDesktopLoginInteractionLevel Interact
 	// Avoid spamming this if the service is down.
 	if (FailedLoginAttempts > UE_BUILDSERVERINSTANCE_MAX_FAILED_LOGIN_ATTEMPTS)
 	{
+		UE_LOG(LogBuildServiceInstance, Display,
+			TEXT("OidcToken: Skipping token refresh due to failed login attempts (%d)."), FailedLoginAttempts);
 		return false;
 	}
 
@@ -737,6 +764,8 @@ FBuildServiceInstance::AcquireAccessToken(EDesktopLoginInteractionLevel Interact
 	// If the token was updated while we waited to take the lock, then it should now be valid.
 	if (Access && Access->GetSerial() > WantsToUpdateTokenSerial)
 	{
+		UE_LOG(LogBuildServiceInstance, Display,
+			TEXT("OidcToken: Skipping token refresh due to token serial change."));
 		return true;
 	}
 
@@ -754,7 +783,7 @@ FBuildServiceInstance::AcquireAccessToken(EDesktopLoginInteractionLevel Interact
 
 		const double ExpiryTimeSeconds = (TokenExpiresAt - FDateTime::UtcNow()).GetTotalSeconds();
 		UE_LOG(LogBuildServiceInstance, Display,
-			TEXT("OidcToken: Logged in to HTTP DDC services. Expires at %s which is in %.0f seconds."),
+			TEXT("OidcToken: Logged in to HTTP Build services. Expires at %s which is in %.0f seconds."),
 			*TokenExpiresAt.ToString(), ExpiryTimeSeconds);
 		SetAccessTokenAndUnlock(Lock, AccessTokenString, ExpiryTimeSeconds);
 		return true;
@@ -793,9 +822,15 @@ FBuildServiceInstance::SetAccessTokenAndUnlock(FScopeLock& Lock, FStringView Tok
 		// Schedule a refresh of the token ahead of expiry time (this will not work in commandlets)
 		if (!IsRunningCommandlet())
 		{
+			UE_LOG(LogBuildServiceInstance, Display,
+				TEXT("OidcToken: Scheduling token refresh in %.0f seconds."),
+				float(FMath::Min(RefreshDelay - RefreshGracePeriod, MAX_flt)));
+
 			RefreshAccessTokenHandle = FTSTicker::GetCoreTicker().AddTicker(FTickerDelegate::CreateLambda(
 				[this](float DeltaTime)
 				{
+					UE_LOG(LogBuildServiceInstance, Display,
+						TEXT("OidcToken: Executing token refresh (no interaction)."));
 					AcquireAccessToken(EDesktopLoginInteractionLevel::None);
 					return false;
 				}
@@ -950,6 +985,7 @@ FBuildServiceInstance::BuildTransferThreadLoop()
 		if (State->bCancelRequested.load(std::memory_order_relaxed))
 		{
 			FWriteScopeLock _(State->Lock);
+			State->TransferEndTime = FMonotonicTimePoint::Now();
 			State->Status = EBuildTransferStatus::Canceled;
 			continue;
 		}
@@ -1037,6 +1073,7 @@ FBuildServiceInstance::BuildTransferThreadLoop()
 			break;
 		}
 
+		UE_LOG(LogBuildServiceInstance, Display, TEXT("Running: \"%s\" %s."), *ZenUtilityPath, *CommandLineArgs);
 		FMonitoredProcess MonitoredUtilityProcess(ZenUtilityPath, *CommandLineArgs, FPaths::GetPath(ZenUtilityPath), true);
 		MonitoredUtilityProcess.OnOutput().BindSPLambda(State, [State](FString Output)
 		{
@@ -1096,14 +1133,19 @@ FBuildServiceInstance::BuildTransferThreadLoop()
 			FString LogFilename = State->GetLogFilename();
 			FOutputDeviceFile::CreateBackupCopy(*LogFilename);
 			State->OutputLogFile = TUniquePtr<FArchive>(IFileManager::Get().CreateFileWriter(*LogFilename, FILEWRITE_Silent | FILEWRITE_AllowRead));
+			State->OutputLogFile->Logf(TEXT("Running: \"%s\" %s"), *ZenUtilityPath, *CommandLineArgs);
+			State->OutputLogFile->Flush();
 			State->Status = EBuildTransferStatus::Active;
+			State->TransferStartTime = FMonotonicTimePoint::Now();
 		}
 
 		if (!MonitoredUtilityProcess.Launch())
 		{
-			UE_LOG(LogBuildServiceInstance, Warning, TEXT("Failed to launch zen utility to download build: '%s'."), *ZenUtilityPath);
+			UE_LOG(LogBuildServiceInstance, Warning, TEXT("Failed to launch zen utility to download build."));
 			FWriteScopeLock _(State->Lock);
+			State->TransferEndTime = FMonotonicTimePoint::Now();
 			State->Status = EBuildTransferStatus::Failed;
+			State->OutputLogFile->Logf(TEXT("Failed to launch."));
 			State->OutputLogFile.Reset();
 			continue;
 		}
@@ -1137,9 +1179,12 @@ FBuildServiceInstance::BuildTransferThreadLoop()
 					FPlatformProcess::Sleep(0.2f);
 				}
 
-				UE_LOG(LogBuildServiceInstance, Display, TEXT("Launch of zen utility for downloading build data: '%s' canceled."), *ZenUtilityPath);
 				FWriteScopeLock _(State->Lock);
+				State->TransferEndTime = FMonotonicTimePoint::Now();
 				State->Status = EBuildTransferStatus::Canceled;
+				double DurationSeconds = (State->TransferEndTime - State->TransferStartTime).ToSeconds();
+				UE_LOG(LogBuildServiceInstance, Display, TEXT("Transfer canceled after 0.1f minutes."), DurationSeconds/60.0f);
+				State->OutputLogFile->Logf(TEXT("Transfer canceled after %.1f minutes."), DurationSeconds/60.0f);
 				State->OutputLogFile.Reset();
 				break;
 			}
@@ -1152,7 +1197,11 @@ FBuildServiceInstance::BuildTransferThreadLoop()
 			{
 				FWriteScopeLock _(State->Lock);
 				State->ReturnCode = MonitoredUtilityProcess.GetReturnCode();
+				State->TransferEndTime = FMonotonicTimePoint::Now();
 				State->Status = EBuildTransferStatus::Succeeded;
+				double DurationSeconds = (State->TransferEndTime - State->TransferStartTime).ToSeconds();
+				UE_LOG(LogBuildServiceInstance, Display, TEXT("Transfer completed successfully after 0.1f minutes."), DurationSeconds/60.0f);
+				State->OutputLogFile->Logf(TEXT("Transfer completed successfully after %.1f minutes."), DurationSeconds/60.0f);
 				State->OutputLogFile.Reset();
 			}
 			else
@@ -1160,7 +1209,10 @@ FBuildServiceInstance::BuildTransferThreadLoop()
 				FWriteScopeLock _(State->Lock);
 				UE_LOG(LogBuildServiceInstance, Warning, TEXT("Unexpected return code after launch of zen utility for downloading build data: '%s' (%d). Output: '%s'"), *ZenUtilityPath, MonitoredUtilityProcess.GetReturnCode(), *State->GetRecentOutput());
 				State->ReturnCode = MonitoredUtilityProcess.GetReturnCode();
+				State->TransferEndTime = FMonotonicTimePoint::Now();
 				State->Status = EBuildTransferStatus::Failed;
+				double DurationSeconds = (State->TransferEndTime - State->TransferStartTime).ToSeconds();
+				State->OutputLogFile->Logf(TEXT("Unexpected return code (%d) from transfer after %.1f minutes"), State->ReturnCode, DurationSeconds/60.0f);
 				State->OutputLogFile.Reset();
 			}
 		}
diff --git a/Engine/Source/Developer/Zen/Public/Experimental/BuildServerInterface.h b/Engine/Source/Developer/Zen/Public/Experimental/BuildServerInterface.h
--- a/Engine/Source/Developer/Zen/Public/Experimental/BuildServerInterface.h
+++ b/Engine/Source/Developer/Zen/Public/Experimental/BuildServerInterface.h
@@ -12,6 +12,7 @@
 #include "HAL/CriticalSection.h"
 #include "HAL/PlatformProcess.h"
 #include "HAL/Thread.h"
+#include "Misc/MonotonicTime.h"
 #include "Serialization/CompactBinary.h"
 #include "Templates/Function.h"
 #include "Templates/PimplPtr.h"
@@ -104,6 +105,8 @@ public:
 		UE_API EBuildTransferStatus GetStatus() const;
 		UE_API FString GetLogFilename() const;
 
+		UE_API FMonotonicTimeSpan GetTransferDuration() const;
+
 		UE_API bool GetCurrentProgress(FString& OutLabel, FString& OutDetail, float& OutPercent) const;
 		UE_API bool GetOverallProgress(FString& OutLabel, FString& OutDetail, float& OutPercent) const;
 
diff --git a/Engine/Source/Editor/StorageServerWidgets/Internal/SBuildSelection.h b/Engine/Source/Editor/StorageServerWidgets/Internal/SBuildSelection.h
--- a/Engine/Source/Editor/StorageServerWidgets/Internal/SBuildSelection.h
+++ b/Engine/Source/Editor/StorageServerWidgets/Internal/SBuildSelection.h
@@ -119,6 +119,7 @@ private:
 	bool ReadSetting(FStringView InKeyName, FString& OutValue);
 
 	void SetUserSelectedProjectDictionaryRoot(FStringView InRoot);
+	void SetFallbackProjectDictionaryRoot(FStringView InRoot);
 
 	static FString SanitizeForPath(const FString& InString);
 	static FString SanitizeForZenId(const FString& InString);
@@ -134,6 +135,7 @@ private:
 
 	TMap<FString, FString> EngineInstallations;
 
+	FString FallbackProjectDictionaryRootDir;
 	TUniquePtr<FUProjectDictionary> FallbackProjectDictionary;
 
 	FString UserSelectedProjectDictionaryRootDir;
diff --git a/Engine/Source/Editor/StorageServerWidgets/Private/SBuildActivity.cpp b/Engine/Source/Editor/StorageServerWidgets/Private/SBuildActivity.cpp
--- a/Engine/Source/Editor/StorageServerWidgets/Private/SBuildActivity.cpp
+++ b/Engine/Source/Editor/StorageServerWidgets/Private/SBuildActivity.cpp
@@ -26,6 +26,7 @@ namespace UE::BuildActivity::Internal
 		const FName ColStatus = TEXT("Status");
 		const FName ColName = TEXT("Name");
 		const FName ColPlatform = TEXT("Platform");
+		const FName ColDuration = TEXT("Duration");
 		const FName ColDescription = TEXT("Description");
 		const FName ColActions = TEXT("Actions");
 	}
@@ -99,8 +100,10 @@ TSharedRef<SWidget> SBuildActivity::GetGridPanel()
 					.FillWidth(0.2f).HAlignCell(HAlign_Left).HAlignHeader(HAlign_Left).VAlignCell(VAlign_Center)
 				+ SHeaderRow::Column(FBuildActivityIds::ColPlatform).DefaultLabel(LOCTEXT("BuildActivity_ActivityPlatform", "Platform"))
 					.FillWidth(0.1f).HAlignCell(HAlign_Left).HAlignHeader(HAlign_Left).VAlignCell(VAlign_Center)
+				+ SHeaderRow::Column(FBuildActivityIds::ColDuration).DefaultLabel(LOCTEXT("BuildActivity_ActivityDuration", "Duration"))
+					.FillWidth(0.05f).HAlignCell(HAlign_Left).HAlignHeader(HAlign_Left).VAlignCell(VAlign_Center)
 				+ SHeaderRow::Column(FBuildActivityIds::ColDescription).DefaultLabel(LOCTEXT("BuildActivity_ActivityDescription", "Description"))
-					.FillWidth(0.54f).HAlignCell(HAlign_Fill).HAlignHeader(HAlign_Center).VAlignCell(VAlign_Center)
+					.FillWidth(0.49f).HAlignCell(HAlign_Fill).HAlignHeader(HAlign_Center).VAlignCell(VAlign_Center)
 				+ SHeaderRow::Column(FBuildActivityIds::ColActions).DefaultLabel(LOCTEXT("BuildActivity_ActivityActions", "Actions"))
 					.FillWidth(0.06f).HAlignCell(HAlign_Fill).HAlignHeader(HAlign_Center).VAlignCell(VAlign_Center)
 			)
@@ -473,6 +476,28 @@ SBuildActivityTableRow::GenerateWidgetForColumn(const FName& ColumnName)
 					.Text(FText::FromString(BuildActivity->Platform))
 				];
 	}
+	else if (ColumnName == FBuildActivityIds::ColDuration)
+	{
+		return SNew(SHorizontalBox)
+			+ SHorizontalBox::Slot()
+			.Padding(1.f, 1.f)
+			[
+				SNew(STextBlock)
+					.TextStyle(&FAppStyle::GetWidgetStyle<FTextBlockStyle>(TEXT("SmallText")))
+					.Text_Lambda([this]()
+						{
+							UE::FMonotonicTimeSpan Span(BuildActivity->Transfer.GetTransferDuration());
+							if (Span.IsZero())
+							{
+								return FText::FromString(FString());
+							}
+							else
+							{
+								return FText::AsTimespan(FTimespan::FromSeconds(Span.ToSeconds()));
+							}
+						})
+			];
+	}
 	else if (ColumnName == FBuildActivityIds::ColDescription)
 	{
 		return SNew(SOverlay)
diff --git a/Engine/Source/Editor/StorageServerWidgets/Private/SBuildSelection.cpp b/Engine/Source/Editor/StorageServerWidgets/Private/SBuildSelection.cpp
--- a/Engine/Source/Editor/StorageServerWidgets/Private/SBuildSelection.cpp
+++ b/Engine/Source/Editor/StorageServerWidgets/Private/SBuildSelection.cpp
@@ -62,7 +62,7 @@ void SBuildSelection::Construct(const FArguments& InArgs)
 
 	KnownBuildTypePatterns = {
 		// These should be made more constinent and perhaps specified in INI
-		MakeKnownBuildTypePattern(TEXT(".*oplog-?(.*)"), EBuildType::Oplog, LOCTEXT("BuildSelection_BuildType_Oplog", "Cooked Data") ),
+		MakeKnownBuildTypePattern(TEXT(".*oplog-?(.*)"), EBuildType::Oplog, LOCTEXT("BuildSelection_BuildType_Oplog", "Cook Snapshot") ),
 		MakeKnownBuildTypePattern(TEXT(".*packaged-?build-?(.*)"), EBuildType::PackagedBuild, LOCTEXT("BuildSelection_BuildType_PackagedBuild", "Packaged Build"), 1),
 		MakeKnownBuildTypePattern(TEXT(".*staged-?(.*?)-?build.*"), EBuildType::StagedBuild, LOCTEXT("BuildSelection_BuildType_StagedBuild", "Staged Build"), 2),
 		MakeKnownBuildTypePattern(TEXT(".*ugs-?pcb-?(.*)"), EBuildType::EditorPreCompiledBinary, LOCTEXT("BuildSelection_BuildType_PreCompiledBinary", "Editor Pre-compiled Binary")),
@@ -112,17 +112,13 @@ void SBuildSelection::Construct(const FArguments& InArgs)
 	for (int32 BuildTypeIndex = 0; BuildTypeIndex < (int32)EBuildType::Count; ++BuildTypeIndex)
 	{
 		ReadSetting(LexToString((EBuildType)BuildTypeIndex), TEXT("UserSelectedDestination"), UserSelectedDestinations[BuildTypeIndex]);
-	}
-
-	for (int32 BuildTypeIndex = 0; BuildTypeIndex < (int32)EBuildType::Count; ++BuildTypeIndex)
-	{
 		ReadSetting(LexToString((EBuildType)BuildTypeIndex), TEXT("UserSelectedIncludeFilter"), UserSelectedTransferFilters[BuildTypeIndex].IncludeFilter);
 		ReadSetting(LexToString((EBuildType)BuildTypeIndex), TEXT("UserSelectedExcludeFilter"), UserSelectedTransferFilters[BuildTypeIndex].ExcludeFilter);
 	}
 
 	for (const TPair<FString, FString>& EngineInstallation : EngineInstallations)
 	{
-		FallbackProjectDictionary = MakeUnique<FUProjectDictionary>(EngineInstallation.Value);
+		SetFallbackProjectDictionaryRoot(EngineInstallation.Value);
 		break;
 	}
 
@@ -1168,6 +1164,7 @@ FString SBuildSelection::GetDefaultDestination() const
 	EBuildType BuildType = GetSelectedBuildType();
 	IPlatformFile& PlatformFile = FPlatformFileManager::Get().GetPlatformFile();
 
+	FString FallbackPath = FPaths::EngineSavedDir();
 	FString ProjectFilename = FUProjectDictionary::GetDefault().GetProjectPathForGame(**SelectedProject);
 
 	if (ProjectFilename.IsEmpty() && UserSelectedProjectDictionary)
@@ -1175,6 +1172,11 @@ FString SBuildSelection::GetDefaultDestination() const
 		ProjectFilename = UserSelectedProjectDictionary->GetProjectPathForGame(**SelectedProject);
 	}
 
+	if (ProjectFilename.IsEmpty() && FallbackProjectDictionary)
+	{
+		ProjectFilename = FallbackProjectDictionary->GetProjectPathForGame(**SelectedProject);
+	}
+
 	if (ProjectFilename.IsEmpty())
 	{
 		if (BuildType == EBuildType::Oplog)
@@ -1182,16 +1184,25 @@ FString SBuildSelection::GetDefaultDestination() const
 			return *SelectedProject;
 		}
 
+		FString DownloadBase;
+		if (!UserSelectedProjectDictionaryRootDir.IsEmpty())
+		{
+			DownloadBase = UserSelectedProjectDictionaryRootDir;
+		}
+		else if (!FallbackProjectDictionaryRootDir.IsEmpty())
+		{
+			DownloadBase = FallbackProjectDictionaryRootDir;
+		}
+		else
+		{
+			DownloadBase = FPaths::EngineSavedDir();
+		}
+
 		return PlatformFile.ConvertToAbsolutePathForExternalAppForRead(
-				*FPaths::Combine(FPaths::EngineSavedDir(), TEXT("DownloadedBuilds"))
+				*FPaths::Combine(DownloadBase, TEXT("DownloadedBuilds"))
 			);
 	}
 
-	if (ProjectFilename.IsEmpty() && FallbackProjectDictionary)
-	{
-		ProjectFilename = FallbackProjectDictionary->GetProjectPathForGame(**SelectedProject);
-	}
-
 	switch (BuildType)
 	{
 	case EBuildType::StagedBuild:
@@ -1285,6 +1296,22 @@ void SBuildSelection::SetUserSelectedProjectDictionaryRoot(const FStringView InR
 	}
 }
 
+void SBuildSelection::SetFallbackProjectDictionaryRoot(const FStringView InRoot)
+{
+	if (InRoot.IsEmpty())
+	{
+		FallbackProjectDictionaryRootDir.Empty();
+		FallbackProjectDictionary.Reset();
+	}
+	FString Root(InRoot);
+	FString FullRoot = FPaths::ConvertRelativePathToFull(Root);
+	if (!FPaths::IsSamePath(FullRoot, FallbackProjectDictionaryRootDir))
+	{
+		FallbackProjectDictionaryRootDir = FullRoot;
+		FallbackProjectDictionary = MakeUnique<FUProjectDictionary>(FallbackProjectDictionaryRootDir);
+	}
+}
+
 FString SBuildSelection::SanitizeForPath(const FString& InString)
 {
 	// TODO: Had to remove path sanitization as the engine only provides validation
diff --git a/Engine/Source/Programs/BuildStorageTool/Private/BuildStorageTool.cpp b/Engine/Source/Programs/BuildStorageTool/Private/BuildStorageTool.cpp
--- a/Engine/Source/Programs/BuildStorageTool/Private/BuildStorageTool.cpp
+++ b/Engine/Source/Programs/BuildStorageTool/Private/BuildStorageTool.cpp
@@ -371,17 +371,31 @@ public:
 		// Setting focus seems to have to happen after the Window has been added
 		Slate.ClearKeyboardFocus(EFocusCause::Cleared);
 
+		double DeltaTime = 0.0;
+		double LastTime = FPlatformTime::Seconds();
+		const float IdealFrameTime = 1.0f / 60;
+		const float BackgroundFrameTime = 1.0f / 4;
+		
 		// loop until the app is ready to quit
 		while (!IsEngineExitRequested())
 		{
 			BeginExitIfRequested();
 
 			FTaskGraphInterface::Get().ProcessThreadUntilIdle(ENamedThreads::GameThread);
-			FTSTicker::GetCoreTicker().Tick(FApp::GetDeltaTime());
+			FTSTicker::GetCoreTicker().Tick(DeltaTime);
 			Slate.PumpMessages();
 			Slate.Tick();
 
-			FPlatformProcess::Sleep(1.0f / 30.0f);
+			const float FrameTime = !FPlatformApplicationMisc::IsThisApplicationForeground() && (FPlatformTime::Seconds() - FSlateApplication::Get().GetLastUserInteractionTime()) > 5 ? BackgroundFrameTime : IdealFrameTime;
+			FPlatformProcess::Sleep(FMath::Max<float>(0.f, FrameTime - static_cast<float>(FPlatformTime::Seconds() - LastTime)));
+
+			DeltaTime = FPlatformTime::Seconds() - LastTime;
+			LastTime = FPlatformTime::Seconds();
+
+			UE::Stats::FStats::AdvanceFrame(false);
+			FCoreDelegates::OnEndFrame.Broadcast();
+
+			GFrameCounter++;
 		}
 
 		// Make sure the window is hidden, because it might take a while for the background thread to finish.
