Added Telemetry Support To Build Storage Tool
Tracked Events
- Connection state
- AccessToken aquisition state
- Download/Transfer state
- ZenServer launch staate

Added an Enabled option to Analytics Provider ( Defaulting to true ) cofniguration so that we can turn Horde providers off for Build Storage Tool
Added DeafultEngine.ini to BuildStorageTool config and set Horde providers Enabled=False
Added Enabled=True for built in providers to BaseEngine.ini
Added Log Display text to report any disabled Providers

#code-analyzed

diff --git a/Engine/Programs/BuildStorageTool/Config/DefaultEngine.ini b/Engine/Programs/BuildStorageTool/Config/DefaultEngine.ini
new file mode 100644
--- /dev/null
+++ b/Engine/Programs/BuildStorageTool/Config/DefaultEngine.ini
@@ -0,0 +1,8 @@
+[StudioTelemetry.Provider.EpicStudio]
+Enabled=True
+
+[StudioTelemetry.Provider.Horde.Dev]
+Enabled=False
+
+[StudioTelemetry.Provider.Horde]
+Enabled=False
\ No newline at end of file
diff --git a/Engine/Source/Developer/Zen/Private/BuildServerInterface.cpp b/Engine/Source/Developer/Zen/Private/BuildServerInterface.cpp
--- a/Engine/Source/Developer/Zen/Private/BuildServerInterface.cpp
+++ b/Engine/Source/Developer/Zen/Private/BuildServerInterface.cpp
@@ -31,6 +31,7 @@
 #include "SocketSubsystem.h"
 #include "String/LexFromString.h"
 #include "Tasks/Task.h"
+#include "StudioTelemetry.h"
 
 #if WITH_SSL
 #include "Ssl.h"
@@ -51,6 +52,7 @@ struct FProgressState
 	FString Label;
 	FString Detail;
 	float Percent = 0.0f;
+	float Bandwidth = 0.0f;
 };
 
 struct FBuildTransferState : public TSharedFromThis<FBuildTransferState, ESPMode::ThreadSafe>
@@ -65,6 +67,7 @@ struct FBuildTransferState : public TSharedFromThis<FBuildTransferState, ESPMode
 	FString IncludeFilter;
 	FString ExcludeFilter;
 	FName TargetPlatformName = NAME_None;
+	double StartTime=0;
 
 	std::atomic<bool> bCancelRequested = false;
 
@@ -259,6 +262,7 @@ FBuildServiceInstance::Connect(bool bAllowInteractive, FOnConnectionComplete&& O
 			EDesktopLoginInteractionLevel InteractionLevel = bAllowInteractive ? EDesktopLoginInteractionLevel::Interactive : EDesktopLoginInteractionLevel::None;
 			if (AcquireAccessToken(InteractionLevel))
 			{
+				RecordConnectionEvent(true);
 				ConnectionState.store(EConnectionState::ConnectionSucceeded, std::memory_order_relaxed);
 				if (OnConnectionComplete)
 				{
@@ -267,7 +271,9 @@ FBuildServiceInstance::Connect(bool bAllowInteractive, FOnConnectionComplete&& O
 			}
 			else
 			{
-				UE_LOG(LogBuildServiceInstance, Warning, TEXT("Unable to acquire access token."));
+				const FString MessageString = TEXT("Unable to acquire access token.");
+				UE_LOG(LogBuildServiceInstance, Warning, TEXT("%s"), *MessageString);
+				RecordConnectionEvent(false, MessageString);
 				ConnectionState.store(EConnectionState::ConnectionFailed, std::memory_order_relaxed);
 				if (OnConnectionComplete)
 				{
@@ -282,7 +288,9 @@ FBuildServiceInstance::Connect(bool bAllowInteractive, FOnConnectionComplete&& O
 			EffectiveDomain.Append(ResolvedHost);
 
 			FString HostCandidates = HostBuilder.GetHostCandidatesString();
-			UE_LOG(LogBuildServiceInstance, Warning, TEXT("Unable to resolve best host candidate to use, most likely none of the suggested hosts was reachable. Attempted hosts were: '%s' ."), *HostCandidates);
+			const FString MessageString = ( TEXT("Unable to resolve best host candidate to use. Most likely none of the suggested hosts was reachable. Attempted hosts were: '%s' ."), *HostCandidates );
+			UE_LOG(LogBuildServiceInstance, Warning, TEXT("%s"), *MessageString);
+			RecordConnectionEvent(false, MessageString);
 			ConnectionState.store(EConnectionState::ConnectionFailed, std::memory_order_relaxed);
 			if (OnConnectionComplete)
 			{
@@ -445,6 +453,7 @@ FBuildServiceInstance::StartBuildTransfer(
 	NewState->Destination = DestinationFolder;
 	NewState->IncludeFilter = IncludeFilter;
 	NewState->ExcludeFilter = ExcludeFilter;
+	NewState->StartTime = FPlatformTime::Seconds();
 	BuildTransferThreadStates.Enqueue(NewState.ToSharedRef());
 	KickBuildTransferThread();
 
@@ -471,6 +480,7 @@ FBuildServiceInstance::StartOplogBuildTransfer(
 	NewState->Namespace = Namespace;
 	NewState->Bucket = Bucket;
 	NewState->Destination = *WriteToString<64>(DestinationProjectId, TEXT("/"), DestinationOplogId);
+	NewState->StartTime = FPlatformTime::Seconds();
 	NewState->ProjectFilePath = ProjectFilePath;
 	NewState->TargetPlatformName = TargetPlatformName;
 	BuildTransferThreadStates.Enqueue(NewState.ToSharedRef());
@@ -498,6 +508,7 @@ FBuildServiceInstance::RepeatBuildTransfer(FBuildTransfer BuildTransfer)
 		NewState->ExcludeFilter = BuildTransfer.State->ExcludeFilter;
 		NewState->ProjectFilePath = BuildTransfer.State->ProjectFilePath;
 		NewState->TargetPlatformName = BuildTransfer.State->TargetPlatformName;
+		NewState->StartTime = FPlatformTime::Seconds();
 	}
 	
 	BuildTransferThreadStates.Enqueue(NewState.ToSharedRef());
@@ -786,17 +797,21 @@ FBuildServiceInstance::AcquireAccessToken(EDesktopLoginInteractionLevel Interact
 			TEXT("OidcToken: Logged in to HTTP Build services. Expires at %s which is in %.0f seconds."),
 			*TokenExpiresAt.ToString(), ExpiryTimeSeconds);
 		SetAccessTokenAndUnlock(Lock, AccessTokenString, ExpiryTimeSeconds);
+		RecordAcquireAccessTokenEvent(true);
+		
 		return true;
 	}
 	else if (DesktopPlatform)
 	{
 		UE_LOG(LogBuildServiceInstance, Warning, TEXT("OidcToken: Failed to log in to HTTP services."));
+		RecordAcquireAccessTokenEvent(false, TEXT("OidcToken: Failed to log in to HTTP services."));
 		FailedLoginAttempts++;
 		return false;
 	}
 	else
 	{
 		UE_LOG(LogBuildServiceInstance, Warning, TEXT("OidcToken: Use of OAuthProviderIdentifier requires that the target depend on DesktopPlatform."));
+		RecordAcquireAccessTokenEvent(false, TEXT("OidcToken: Use of OAuthProviderIdentifier requires that the target depend on DesktopPlatform."));
 		FailedLoginAttempts++;
 		return false;
 	}
@@ -887,7 +902,10 @@ FBuildServiceInstance::InvokeZenUtilitySync(FStringView InCommandLineArgs)
 	FMonitoredProcess MonitoredUtilityProcess(AbsoluteUtilityPath, *CommandLineArgs, FPaths::GetPath(AbsoluteUtilityPath), true);
 	if (!MonitoredUtilityProcess.Launch())
 	{
-		UE_LOG(LogBuildServiceInstance, Warning, TEXT("Failed to launch zen utility to gather build data: '%s'."), *AbsoluteUtilityPath);
+		const FString MessageString = FString::Printf( TEXT("Failed to launch zen utility to gather build data in path '%s'."), *AbsoluteUtilityPath );
+		UE_LOG(LogBuildServiceInstance, Warning, TEXT("%s"), *MessageString);
+		RecordZenServerEvent(false, MessageString);
+		
 		return false;
 	}
 
@@ -898,8 +916,10 @@ FBuildServiceInstance::InvokeZenUtilitySync(FStringView InCommandLineArgs)
 		if (Duration > 120.0)
 		{
 			MonitoredUtilityProcess.Cancel(true);
-			UE_LOG(LogBuildServiceInstance, Warning, TEXT("Cancelled launch of zen utility for gathering build data: '%s' due to timeout."), *AbsoluteUtilityPath);
-
+			const FString MessageString = FString::Printf( TEXT("Cancelled launch of zen utility for gathering build data: '%s' due to timeout."), *AbsoluteUtilityPath);
+			UE_LOG(LogBuildServiceInstance, Warning, TEXT("%s"), *MessageString);
+			RecordZenServerEvent(false, MessageString);
+			
 			// Wait for execution to be terminated
 			while (MonitoredUtilityProcess.Update())
 			{
@@ -923,6 +943,7 @@ FBuildServiceInstance::InvokeZenUtilitySync(FStringView InCommandLineArgs)
 	if (MonitoredUtilityProcess.GetReturnCode() != 0)
 	{
 		UE_LOG(LogBuildServiceInstance, Warning, TEXT("Unexpected return code after launch of zen utility for gathering build data: '%s' (%d). Output: '%s'"), *AbsoluteUtilityPath, MonitoredUtilityProcess.GetReturnCode(), *OutputString);
+		RecordZenServerEvent(false, TEXT("Unexpected return code after launch of zen utility for gathering build data"));
 		return false;
 	}
 
@@ -1186,6 +1207,8 @@ FBuildServiceInstance::BuildTransferThreadLoop()
 				UE_LOG(LogBuildServiceInstance, Display, TEXT("Transfer canceled after 0.1f minutes."), DurationSeconds/60.0f);
 				State->OutputLogFile->Logf(TEXT("Transfer canceled after %.1f minutes."), DurationSeconds/60.0f);
 				State->OutputLogFile.Reset();
+				RecordTransferStateEvent(State);
+
 				break;
 			}
 			FPlatformProcess::Sleep(0.1f);
@@ -1203,6 +1226,7 @@ FBuildServiceInstance::BuildTransferThreadLoop()
 				UE_LOG(LogBuildServiceInstance, Display, TEXT("Transfer completed successfully after 0.1f minutes."), DurationSeconds/60.0f);
 				State->OutputLogFile->Logf(TEXT("Transfer completed successfully after %.1f minutes."), DurationSeconds/60.0f);
 				State->OutputLogFile.Reset();
+				RecordTransferStateEvent(State);
 			}
 			else
 			{
@@ -1214,6 +1238,7 @@ FBuildServiceInstance::BuildTransferThreadLoop()
 				double DurationSeconds = (State->TransferEndTime - State->TransferStartTime).ToSeconds();
 				State->OutputLogFile->Logf(TEXT("Unexpected return code (%d) from transfer after %.1f minutes"), State->ReturnCode, DurationSeconds/60.0f);
 				State->OutputLogFile.Reset();
+				RecordTransferStateEvent(State);
 			}
 		}
 
@@ -1248,4 +1273,109 @@ FBuildServiceInstance::BuildTransferThreadLoop()
 	}
 }
 
+void
+FBuildServiceInstance::RecordAcquireAccessTokenEvent(bool Succeeded, const FString& MessageString)
+{
+	// Generate and send Telemetry Event for accessing the Access Token
+	if (FStudioTelemetry::IsAvailable())
+	{
+		TArray<FAnalyticsEventAttribute> Attributes;
+		const int SchemaVersion = 1;
+		Attributes.Emplace(TEXT("SchemaVersion"), SchemaVersion);
+		Attributes.Emplace(TEXT("Host"), Settings.GetHost());
+		Attributes.Emplace(TEXT("Outcome"), Succeeded ? TEXT("Succeeded") : TEXT("Failed"));
+		Attributes.Emplace(TEXT("Message"), MessageString);
+		Attributes.Emplace(TEXT("FailedAttempts"), FailedLoginAttempts);
+		FStudioTelemetry::Get().RecordEvent(TEXT("Core.BuildService.AquireAccessToken"), Attributes);
+		FStudioTelemetry::Get().FlushEvents();
+	}
+}
+
+void FBuildServiceInstance::RecordConnectionEvent(bool Succeeded, const FString& MessageString)
+{
+	// Generate and send Telemetry Event for connection
+	if (FStudioTelemetry::IsAvailable())
+	{
+		TArray<FAnalyticsEventAttribute> Attributes;
+		const int SchemaVersion = 1;
+		Attributes.Emplace(TEXT("SchemaVersion"), SchemaVersion);
+		Attributes.Emplace(TEXT("Host"), Settings.GetHost());
+		Attributes.Emplace(TEXT("Outcome"), Succeeded ? TEXT("Succeeded") : TEXT("Failed"));
+		Attributes.Emplace(TEXT("Message"), MessageString);
+		FStudioTelemetry::Get().RecordEvent(TEXT("Core.BuildService.Connect"), Attributes);
+		FStudioTelemetry::Get().FlushEvents();
+	}
+}
+
+void
+FBuildServiceInstance::RecordZenServerEvent(bool Succeeded, const FString& MessageString)
+{
+	// Generate and send Telemetry Event for interactions with the Zen Server
+	if (FStudioTelemetry::IsAvailable())
+	{
+		TArray<FAnalyticsEventAttribute> Attributes;
+		const int SchemaVersion = 1;
+		Attributes.Emplace(TEXT("SchemaVersion"), SchemaVersion);
+		Attributes.Emplace(TEXT("Host"), Settings.GetHost());
+		Attributes.Emplace(TEXT("Outcome"), Succeeded ? TEXT("Succeeded") : TEXT("Failed"));
+		Attributes.Emplace(TEXT("Message"), MessageString);
+		FStudioTelemetry::Get().RecordEvent(TEXT("Core.BuildService.Zen"), Attributes);
+		FStudioTelemetry::Get().FlushEvents();
+	}
+}
+
+void
+FBuildServiceInstance::RecordTransferStateEvent(TSharedRef<Private::FBuildTransferState> State)
+{
+	// Generate and send Telemetry Event for the transfer state
+	if (FStudioTelemetry::IsAvailable())
+	{
+		TArray<FAnalyticsEventAttribute> Attributes;
+		const int SchemaVersion = 1;
+
+		FString Status;
+		switch (State->Status)
+		{
+		case EBuildTransferStatus::Canceled:
+		{
+			Status = TEXT("Cancelled");
+			break;
+		}
+		case EBuildTransferStatus::Failed:
+		{
+			Status = TEXT("Failed");
+			break;
+		}
+		case EBuildTransferStatus::Succeeded:
+		{
+			Status = TEXT("Succeeded");
+			break;
+		}
+		case EBuildTransferStatus::Queued:
+		{
+			Status = TEXT("Queued");
+			break;
+		}
+		default:
+		{
+			Status = TEXT("Undefined");
+			break;
+		}
+		}
+
+		Attributes.Emplace(TEXT("SchemaVersion"), SchemaVersion);
+		Attributes.Emplace(TEXT("Host"), Settings.GetHost());
+		Attributes.Emplace(TEXT("Status"), Status);
+		Attributes.Emplace(TEXT("Name"), State->Name);
+		Attributes.Emplace(TEXT("Namespace"), State->Namespace);
+		Attributes.Emplace(TEXT("Bucket"), State->Bucket);
+		Attributes.Emplace(TEXT("ReturnCode"), State->ReturnCode);
+		Attributes.Emplace(TEXT("Platform"), State->TargetPlatformName);
+		Attributes.Emplace(TEXT("Duration"), FPlatformTime::Seconds() - State->StartTime);
+
+		FStudioTelemetry::Get().RecordEvent(TEXT("Core.BuildService.Transfer"), Attributes);
+		FStudioTelemetry::Get().FlushEvents();
+	}
+}
+
 } // namespace UE::StorageService::Build
diff --git a/Engine/Source/Developer/Zen/Public/Experimental/BuildServerInterface.h b/Engine/Source/Developer/Zen/Public/Experimental/BuildServerInterface.h
--- a/Engine/Source/Developer/Zen/Public/Experimental/BuildServerInterface.h
+++ b/Engine/Source/Developer/Zen/Public/Experimental/BuildServerInterface.h
@@ -169,6 +169,11 @@ private:
 	bool InvokeZenUtilitySync(FStringView CommandlineArgs);
 	void InvokeZenUtility(FStringView CommandlineArgs, FPreZenUtilityInvocation&& PreInvocation, FOnZenUtilityInvocationComplete&& OnComplete);
 	void CallOnRefreshNamespacesAndBucketsComplete(FOnRefreshNamespacesAndBucketsComplete&& InOnRefreshNamespacesAndBucketsComplete);
+	
+	void RecordAcquireAccessTokenEvent(bool Succeeded, const FString& MessageString = {});
+	void RecordConnectionEvent(bool Succeeded, const FString& MessageString = {});
+	void RecordZenServerEvent(bool Succeeded, const FString& MessageString = {});
+	void RecordTransferStateEvent(TSharedRef<Private::FBuildTransferState> State);
 
 	void KickBuildTransferThread();
 	void BuildTransferThreadLoop();
diff --git a/Engine/Source/Developer/Zen/Zen.Build.cs b/Engine/Source/Developer/Zen/Zen.Build.cs
--- a/Engine/Source/Developer/Zen/Zen.Build.cs
+++ b/Engine/Source/Developer/Zen/Zen.Build.cs
@@ -7,7 +7,7 @@ public class Zen : ModuleRules
 {
 	public Zen(ReadOnlyTargetRules Target) : base(Target)
 	{
-		PrivateDependencyModuleNames.AddRange(new string[] { "Core", "Sockets", "SSL", "Json", "DevHttp" });
+		PrivateDependencyModuleNames.AddRange(new string[] { "Core", "Sockets", "SSL", "Json", "DevHttp",  "StudioTelemetry" });
 		PrivateIncludePathModuleNames.AddRange(new string[] { "DesktopPlatform", "Analytics" });
 		
 		if (Target.Platform.IsInGroup(UnrealPlatformGroup.Windows) && Target.WindowsPlatform.bUseXCurl)
diff --git a/Engine/Source/Programs/BuildStorageTool/BuildStorageTool.Build.cs b/Engine/Source/Programs/BuildStorageTool/BuildStorageTool.Build.cs
--- a/Engine/Source/Programs/BuildStorageTool/BuildStorageTool.Build.cs
+++ b/Engine/Source/Programs/BuildStorageTool/BuildStorageTool.Build.cs
@@ -12,6 +12,7 @@ public class BuildStorageTool : ModuleRules
 			new string[]
 			{
 				"ApplicationCore",
+				"Analytics",
 				"Core",
 				"DesktopPlatform",
 				"Json",
@@ -21,6 +22,7 @@ public class BuildStorageTool : ModuleRules
 				"SlateCore",
 				"StandaloneRenderer",
 				"StorageServerWidgets",
+				"StudioTelemetry",
 				"Zen",
 				"OutputLog"
 			});
diff --git a/Engine/Source/Programs/BuildStorageTool/Private/BuildStorageTool.cpp b/Engine/Source/Programs/BuildStorageTool/Private/BuildStorageTool.cpp
--- a/Engine/Source/Programs/BuildStorageTool/Private/BuildStorageTool.cpp
+++ b/Engine/Source/Programs/BuildStorageTool/Private/BuildStorageTool.cpp
@@ -37,6 +37,8 @@
 #include "OutputLogCreationParams.h"
 #include "OutputLogModule.h"
 #include "OutputLogSettings.h"
+#include "StudioTelemetry.h"
+
 
 #if PLATFORM_WINDOWS
 #include "Runtime/Launch/Resources/Windows/Resource.h"
@@ -300,6 +302,8 @@ public:
 
 	void Run()
 	{
+		STUDIO_TELEMETRY_SESSION_SCOPE
+
 		Window =
 			SNew(SWindow)
 			.Title(GetWindowTitle())
diff --git a/Engine/Source/Runtime/Analytics/AnalyticsHorde/AnalyticsHorde.Build.cs b/Engine/Source/Runtime/Analytics/AnalyticsHorde/AnalyticsHorde.Build.cs
--- a/Engine/Source/Runtime/Analytics/AnalyticsHorde/AnalyticsHorde.Build.cs
+++ b/Engine/Source/Runtime/Analytics/AnalyticsHorde/AnalyticsHorde.Build.cs
@@ -9,7 +9,6 @@ namespace UnrealBuildTool.Rules
 			PrivateDependencyModuleNames.AddRange(
 				new string[] {
 				"Core",
-				"CoreUObject",
 				"Analytics",
 				"AnalyticsET"
 				}
diff --git a/Engine/Source/Runtime/Analytics/AnalyticsLog/AnalyticsLog.Build.cs b/Engine/Source/Runtime/Analytics/AnalyticsLog/AnalyticsLog.Build.cs
--- a/Engine/Source/Runtime/Analytics/AnalyticsLog/AnalyticsLog.Build.cs
+++ b/Engine/Source/Runtime/Analytics/AnalyticsLog/AnalyticsLog.Build.cs
@@ -14,7 +14,6 @@ namespace UnrealBuildTool.Rules
 
 			PrivateDependencyModuleNames.AddRange(
 				new string[] {
-				"CoreUObject",
 				"Analytics"
 				}
 			);
diff --git a/Engine/Source/Runtime/StudioTelemetry/Private/AnalyticsProviderBroadcast.cpp b/Engine/Source/Runtime/StudioTelemetry/Private/AnalyticsProviderBroadcast.cpp
--- a/Engine/Source/Runtime/StudioTelemetry/Private/AnalyticsProviderBroadcast.cpp
+++ b/Engine/Source/Runtime/StudioTelemetry/Private/AnalyticsProviderBroadcast.cpp
@@ -53,6 +53,16 @@ FAnalyticsProviderBroadcast::FAnalyticsProviderBroadcast()
 
 				FString UsageType;
 
+				// Is this Provider enabled? A given application may choose to disable it 
+				bool Enabled = true;
+				Enabled = GConfig->GetBoolOrDefault(*ProviderSection, TEXT("Enabled"), true, GEngineIni);
+
+				if (Enabled == false)
+				{
+					UE_LOG(LogAnalytics, Display, TEXT("Provider %s is disabled for this application"), *ProviderSection);
+					continue;
+				}
+
 				// Validate the usage type is for this build type
 				if (GConfig->GetString(*ProviderSection, TEXT("UsageType"), UsageType, GEngineIni))
 				{
@@ -100,7 +110,7 @@ FAnalyticsProviderBroadcast::FAnalyticsProviderBroadcast()
 					UE_LOG(LogAnalytics, Warning, TEXT("There must be a valid UsageType specified for analytics provider %s"), *ProviderSection);
 					continue;
 				}
-
+				
 				FString ProviderModuleName;
 
 				if (GConfig->GetString(*ProviderSection, TEXT("ProviderModule"), ProviderModuleName, GEngineIni))
diff --git a/Engine/Source/Runtime/StudioTelemetry/StudioTelemetry.build.cs b/Engine/Source/Runtime/StudioTelemetry/StudioTelemetry.build.cs
--- a/Engine/Source/Runtime/StudioTelemetry/StudioTelemetry.build.cs
+++ b/Engine/Source/Runtime/StudioTelemetry/StudioTelemetry.build.cs
@@ -13,8 +13,6 @@ public class StudioTelemetry : ModuleRules
 			new string[] {
 				"Core",
 				"HTTP",
-				"CoreUObject",
-                "EngineSettings",
                 "BuildSettings",
                 "Analytics",
                 "AnalyticsET",
