[Zen]
* Removed usage of time()/localtime() to make file compile with header units (without modifying windows sdk)

diff --git a/Engine/Source/Developer/Zen/Private/ZenServerInterface.cpp b/Engine/Source/Developer/Zen/Private/ZenServerInterface.cpp
--- a/Engine/Source/Developer/Zen/Private/ZenServerInterface.cpp
+++ b/Engine/Source/Developer/Zen/Private/ZenServerInterface.cpp
@@ -17,7 +17,6 @@
 #include "HAL/Platform.h"
 #include "HAL/PlatformMisc.h"
 #include "HAL/PlatformProcess.h"
-#include "HAL/PlatformTime.h"
 #include "Logging/LogScopedCategoryAndVerbosityOverride.h"
 #include "Logging/StructuredLog.h"
 #include "Misc/App.h"
@@ -328,15 +327,9 @@ GetInstallCrashpadHandlerFilePath(const FString& InTreePath)
 
 FString WriteMiniDump(HANDLE ProcessHandle, const FString& ExeName, const FString& OutputFolderPath)
 {
-	time_t rawtime;
-	struct tm timeinfo;
-
-	time(&rawtime);
-	localtime_s(&timeinfo, &rawtime);
-
-	FString FileName;
-	FileName.Appendf(TEXT("%s_Dump_%04d%02d%02d_%02d%02d%02d.dmp"),
-		*ExeName, timeinfo.tm_year + 1900, timeinfo.tm_mon, timeinfo.tm_mday, timeinfo.tm_hour, timeinfo.tm_min, timeinfo.tm_sec);
+	const FDateTime Now = FDateTime::Now();
+	const FString Timestamp = Now.ToString(TEXT("%Y%m%d_%H%M%S"));
+	FString FileName = FString::Printf(TEXT("%s_Dump_%s.dmp"), *ExeName, *Timestamp);
 
 	FString DumpPath = FPaths::Combine(OutputFolderPath, FileName);
 
