Improve strictness around determining the project file path, root dir, and engine dir when importing snapshots in BuildStorageTool.

diff --git a/Engine/Source/Developer/Zen/Private/BuildServerInterface.cpp b/Engine/Source/Developer/Zen/Private/BuildServerInterface.cpp
--- a/Engine/Source/Developer/Zen/Private/BuildServerInterface.cpp
+++ b/Engine/Source/Developer/Zen/Private/BuildServerInterface.cpp
@@ -28,6 +28,8 @@
 #include "Serialization/CompactBinarySerialization.h"
 #include "Serialization/CompactBinaryValidation.h"
 #include "Serialization/CompactBinaryWriter.h"
+#include "Serialization/JsonReader.h"
+#include "Serialization/JsonSerializer.h"
 #include "SocketSubsystem.h"
 #include "String/LexFromString.h"
 #include "Tasks/Task.h"
@@ -63,6 +65,8 @@ struct FBuildTransferState : public TSharedFromThis<FBuildTransferState, ESPMode
 	FString Namespace;
 	FString Bucket;
 	FString Destination;
+	FString RootDir;
+	FString EngineDir;
 	FString ProjectFilePath;
 	FString IncludeFilter;
 	FString ExcludeFilter;
@@ -478,6 +482,8 @@ FBuildServiceInstance::StartOplogBuildTransfer(
 	FStringView Name,
 	FStringView DestinationProjectId,
 	FStringView DestinationOplogId,
+	FStringView RootDir,
+	FStringView EngineDir,
 	FStringView ProjectFilePath,
 	FStringView Namespace,
 	FStringView Bucket,
@@ -491,6 +497,8 @@ FBuildServiceInstance::StartOplogBuildTransfer(
 	NewState->Bucket = Bucket;
 	NewState->Destination = *WriteToString<64>(DestinationProjectId, TEXT("/"), DestinationOplogId);
 	NewState->StartTime = FPlatformTime::Seconds();
+	NewState->RootDir = RootDir;
+	NewState->EngineDir = EngineDir;
 	NewState->ProjectFilePath = ProjectFilePath;
 	NewState->TargetPlatformName = TargetPlatformName;
 	BuildTransferThreadStates.Enqueue(NewState.ToSharedRef());
@@ -516,6 +524,8 @@ FBuildServiceInstance::RepeatBuildTransfer(FBuildTransfer BuildTransfer)
 		NewState->Destination = BuildTransfer.State->Destination;
 		NewState->IncludeFilter = BuildTransfer.State->IncludeFilter;
 		NewState->ExcludeFilter = BuildTransfer.State->ExcludeFilter;
+		NewState->RootDir = BuildTransfer.State->RootDir;
+		NewState->EngineDir = BuildTransfer.State->EngineDir;
 		NewState->ProjectFilePath = BuildTransfer.State->ProjectFilePath;
 		NewState->TargetPlatformName = BuildTransfer.State->TargetPlatformName;
 		NewState->StartTime = FPlatformTime::Seconds();
@@ -910,12 +920,17 @@ FBuildServiceInstance::AddZenUtilityBuildServerArguments(const TCHAR* InCommandL
 }
 
 bool
-FBuildServiceInstance::InvokeZenUtilitySync(FStringView InCommandLineArgs)
+FBuildServiceInstance::InvokeZenUtilitySync(FStringView InCommandLineArgs, FString* OutputText)
 {
 	FString ZenUtilityPath = UE::Zen::GetLocalInstallUtilityPath();
 
 	FString CommandLineArgs(InCommandLineArgs);
 
+	if (OutputText)
+	{
+		OutputText->Empty();
+	}
+
 	FString AbsoluteUtilityPath = FPaths::ConvertRelativePathToFull(ZenUtilityPath);
 	FMonitoredProcess MonitoredUtilityProcess(AbsoluteUtilityPath, *CommandLineArgs, FPaths::GetPath(AbsoluteUtilityPath), true);
 	if (!MonitoredUtilityProcess.Launch())
@@ -965,6 +980,11 @@ FBuildServiceInstance::InvokeZenUtilitySync(FStringView InCommandLineArgs)
 		return false;
 	}
 
+	if (OutputText)
+	{
+		*OutputText = OutputString;
+	}
+
 	return true;
 }
 
@@ -1062,9 +1082,48 @@ FBuildServiceInstance::BuildTransferThreadLoop()
 					}
 				}
 
-				if (!State->ProjectFilePath.IsEmpty())
+				FString RootDir;
+				FString EngineDir;
+				FString ProjectFilePath;
+
+				// Try to derive the project file path from destination project in zenserver if possible
+				FString ProjectInfoText;
+				if (InvokeZenUtilitySync(FString::Printf(TEXT("project-info --project \"%s\""),
+					*DestinationProjectId), &ProjectInfoText))
+				{
+					int32 StartOfJson = INDEX_NONE;
+					if (ProjectInfoText.FindChar(TEXT('{'), StartOfJson))
+					{
+						ProjectInfoText.RightChopInline(StartOfJson);
+
+						TSharedPtr<FJsonObject> ProjectInfoObject;
+						TSharedRef<TJsonReader<>> Reader = TJsonReaderFactory<>::Create(ProjectInfoText);
+						if (FJsonSerializer::Deserialize(Reader, ProjectInfoObject) && ProjectInfoObject.IsValid())
+						{
+							ProjectInfoObject->TryGetStringField(TEXT("root"), RootDir);
+							ProjectInfoObject->TryGetStringField(TEXT("engine"), EngineDir);
+							ProjectInfoObject->TryGetStringField(TEXT("projectfile"), ProjectFilePath);
+						}
+					}
+				}
+
+				// Only use the specified root/engine/project file path if we didn't determine one from the destination project in zenserver
+				if (RootDir.IsEmpty())
+				{
+					RootDir = State->RootDir;
+				}
+				if (EngineDir.IsEmpty())
+				{
+					EngineDir = State->EngineDir;
+				}
+				if (ProjectFilePath.IsEmpty())
+				{
+					ProjectFilePath = State->ProjectFilePath;
+				}
+
+				if (!ProjectFilePath.IsEmpty())
 				{
-					DestinationDirectory = FPaths::Combine(FPaths::GetPath(State->ProjectFilePath),
+					DestinationDirectory = FPaths::Combine(FPaths::GetPath(ProjectFilePath),
 														TEXT("Saved"),
 														TEXT("Cooked"),
 														DestinationOplogId);
@@ -1075,22 +1134,20 @@ FBuildServiceInstance::BuildTransferThreadLoop()
 					Zen::FZenProjectStoreWriter ProjectStoreWriter(*OplogLifetimeMarkerPath);
 					ProjectStoreWriter.Write(TEXT("[::1]"), LocalPort, true, DestinationProjectId, DestinationOplogId, State->TargetPlatformName);
 
-					FString RootDir = FPaths::RootDir();
-					FString EngineDir = FPaths::EngineDir();
 					FPaths::NormalizeDirectoryName(EngineDir);
-					FString ProjectDir = FPaths::GetPath(State->ProjectFilePath);
+					FString ProjectDir = FPaths::GetPath(ProjectFilePath);
 					FPaths::NormalizeDirectoryName(ProjectDir);
-					FString ProjectPath = State->ProjectFilePath;
+					FString ProjectPath = ProjectFilePath;
 					FPaths::NormalizeFilename(ProjectPath);
 
 					IPlatformFile& PlatformFile = FPlatformFileManager::Get().GetPlatformFile();
 					FString AbsServerRoot = PlatformFile.ConvertToAbsolutePathForExternalAppForRead(*RootDir);
 					FString AbsEngineDir = PlatformFile.ConvertToAbsolutePathForExternalAppForRead(*EngineDir);
 					FString AbsProjectDir = PlatformFile.ConvertToAbsolutePathForExternalAppForRead(*ProjectDir);
-					FString ProjectFilePath = PlatformFile.ConvertToAbsolutePathForExternalAppForRead(*ProjectPath);
+					FString AbsProjectFilePath = PlatformFile.ConvertToAbsolutePathForExternalAppForRead(*ProjectPath);
 
 					InvokeZenUtilitySync(FString::Printf(TEXT("project-create --project \"%s\" --rootdir \"%s\" --enginedir \"%s\" --projectdir \"%s\" --projectfile \"%s\""),
-						*DestinationProjectId, *AbsServerRoot, *AbsEngineDir, *AbsProjectDir, *ProjectFilePath));
+						*DestinationProjectId, *AbsServerRoot, *AbsEngineDir, *AbsProjectDir, *AbsProjectFilePath));
 
 					InvokeZenUtilitySync(FString::Printf(TEXT("oplog-create --project \"%s\" --oplog \"%s\" --gcpath \"%s\""),
 						*DestinationProjectId, *DestinationOplogId, *OplogLifetimeMarkerPath));
diff --git a/Engine/Source/Developer/Zen/Public/Experimental/BuildServerInterface.h b/Engine/Source/Developer/Zen/Public/Experimental/BuildServerInterface.h
--- a/Engine/Source/Developer/Zen/Public/Experimental/BuildServerInterface.h
+++ b/Engine/Source/Developer/Zen/Public/Experimental/BuildServerInterface.h
@@ -152,6 +152,8 @@ public:
 		FStringView Name,
 		FStringView DestinationProjectId,
 		FStringView DestinationOplogId,
+		FStringView RootDir,
+		FStringView EngineDir,
 		FStringView ProjectFilePath,
 		FStringView Namespace,
 		FStringView Bucket,
@@ -167,7 +169,7 @@ private:
 	void SetAccessTokenAndUnlock(FScopeLock &Lock, FStringView Token, double RefreshDelay = 0.0);
 	[[nodiscard]] FString GetAccessToken() const;
 	FString AddZenUtilityBuildServerArguments(const TCHAR* InCommandLineArgs) const;
-	bool InvokeZenUtilitySync(FStringView CommandlineArgs);
+	bool InvokeZenUtilitySync(FStringView CommandlineArgs, FString* Output = nullptr);
 	void InvokeZenUtility(FStringView CommandlineArgs, FPreZenUtilityInvocation&& PreInvocation, FOnZenUtilityInvocationComplete&& OnComplete);
 	void CallOnRefreshNamespacesAndBucketsComplete(FOnRefreshNamespacesAndBucketsComplete&& InOnRefreshNamespacesAndBucketsComplete);
 	
diff --git a/Engine/Source/Editor/StorageServerWidgets/Private/SBuildSelection.cpp b/Engine/Source/Editor/StorageServerWidgets/Private/SBuildSelection.cpp
--- a/Engine/Source/Editor/StorageServerWidgets/Private/SBuildSelection.cpp
+++ b/Engine/Source/Editor/StorageServerWidgets/Private/SBuildSelection.cpp
@@ -1447,7 +1447,23 @@ void SBuildSelection::DownloadSelectedArtifactsInSelectedBuildGroup()
 
 				if (BuildType == EBuildType::Oplog)
 				{
-					FString ProjectFilePath = FUProjectDictionary::GetDefault().GetProjectPathForGame(**SelectedProject);
+					FString RootDir = FPaths::RootDir();
+					FString EngineDir = FPaths::EngineDir();
+					FString ProjectFilename = FUProjectDictionary::GetDefault().GetProjectPathForGame(**SelectedProject);
+
+					if (ProjectFilename.IsEmpty() && UserSelectedProjectDictionary)
+					{
+						RootDir = UserSelectedProjectDictionaryRootDir;
+						EngineDir = FPaths::Combine(RootDir, TEXT("Engine"));
+						ProjectFilename = UserSelectedProjectDictionary->GetProjectPathForGame(**SelectedProject);
+					}
+
+					if (ProjectFilename.IsEmpty() && FallbackProjectDictionary)
+					{
+						RootDir = FallbackProjectDictionaryRootDir;
+						EngineDir = FPaths::Combine(RootDir, TEXT("Engine"));
+						ProjectFilename = FallbackProjectDictionary->GetProjectPathForGame(**SelectedProject);
+					}
 					FString DestinationProjectId = GetEffectiveDestination();
 					FString DestinationOplogId = DestinationPlatformName;
 					if (bAppendBuildNameToDestination)
@@ -1455,7 +1471,7 @@ void SBuildSelection::DownloadSelectedArtifactsInSelectedBuildGroup()
 						DestinationOplogId = FString::Printf(TEXT("%s.%s"), *DestinationOplogId, *SanitizeForZenId(SelectedItems[0]->DisplayName));
 					}
 					FBuildServiceInstance::FBuildTransfer BuildTransfer =
-						ServiceInstance->StartOplogBuildTransfer(BuildRecord->BuildId, TransferName, DestinationProjectId, DestinationOplogId, ProjectFilePath, SelectedItems[0]->Namespace, BuildRecord->BucketId, TargetPlatformName);
+						ServiceInstance->StartOplogBuildTransfer(BuildRecord->BuildId, TransferName, DestinationProjectId, DestinationOplogId, RootDir, EngineDir, ProjectFilename, SelectedItems[0]->Namespace, BuildRecord->BucketId, TargetPlatformName);
 					OnBuildTransferStarted.ExecuteIfBound(BuildTransfer, SelectedItems[0]->DisplayName, ArtifactForSelectedGroup);
 				}
 				else
