Fix passthrough args handling in zen service install on Windows.

[RN] none

diff --git a/Engine/Source/Developer/Zen/Private/ZenServerInterface.cpp b/Engine/Source/Developer/Zen/Private/ZenServerInterface.cpp
--- a/Engine/Source/Developer/Zen/Private/ZenServerInterface.cpp
+++ b/Engine/Source/Developer/Zen/Private/ZenServerInterface.cpp
@@ -475,20 +475,25 @@ static bool RunZenUtility(const FString& UtilityPath, const FString& CommandLine
 	}
 }
 
-static bool RunZenUtilityElevated(const FString& UtilityPath, const FString& CommandLine, FString& Output, int& ReturnCode)
+static bool RunZenUtilityElevated(const FString& UtilityPath, const FString& CommandLine, FString& Output, int& ReturnCode, const FString& PassthroughArgs = TEXT(""))
 {
 	FString ElevatedCommandLine;
 	FString ElevatedUtilityPath;
 
 #if PLATFORM_WINDOWS
 	ElevatedUtilityPath = UtilityPath;
-	ElevatedCommandLine = TEXT("--allow-elevation ") + CommandLine;
+	ElevatedCommandLine = CommandLine + TEXT(" --allow-elevation ");
 #endif
 #if PLATFORM_LINUX
 	ElevatedCommandLine = UtilityPath + TEXT(" ") + CommandLine;
 	ElevatedUtilityPath = TEXT("pkexec");
 #endif
 
+	if (PassthroughArgs != TEXT(""))
+	{
+		ElevatedCommandLine += TEXT(" -- ") + PassthroughArgs;
+	}
+
 	return RunZenUtility(ElevatedUtilityPath, ElevatedCommandLine, Output, ReturnCode);
 }
 
@@ -3060,9 +3065,8 @@ bool FZenServiceInstance::ConditionalUpdateSystemServiceInstall(const FSystemSer
 #if PLATFORM_LINUX
 		ServiceInstallCommandLine.Appendf(TEXT(" --user %s"), FPlatformProcess::UserName(false));
 #endif
-		ServiceInstallCommandLine.Appendf(TEXT(" -- %s"), *DesiredCommandLine);
 
-		if (!RunZenUtilityElevated(InTreeUtilityPath, ServiceInstallCommandLine, Output, ReturnCode))
+		if (!RunZenUtilityElevated(InTreeUtilityPath, ServiceInstallCommandLine, Output, ReturnCode, DesiredCommandLine))
 		{
 			PromptUserOfFailedSystemServiceInstall();
 			return false;
