Implement versioning and updating of Zen plugin installations.

[RN] none

diff --git a/Engine/Source/Developer/Zen/Private/ZenServerInterface.cpp b/Engine/Source/Developer/Zen/Private/ZenServerInterface.cpp
--- a/Engine/Source/Developer/Zen/Private/ZenServerInterface.cpp
+++ b/Engine/Source/Developer/Zen/Private/ZenServerInterface.cpp
@@ -247,6 +247,12 @@ GetServicePluginsConfigPath()
 	return FPaths::Combine(GetLocalInstallPath(), FString::Printf(TEXT("zen_plugins_v%d.json"), PluginsConfigVersion));
 }
 
+static FString
+GetServicePluginsVersionInfoPath()
+{
+	return FPaths::Combine(GetLocalInstallPath(), TEXT("zen_plugin_versions.json"));
+}
+
 static void
 CleanOutOfDateServicePluginConfigs()
 {
@@ -1123,6 +1129,8 @@ bool FServicePluginSettings::ReadFromConfig(const FString& InPluginName)
 				UE_LOG(LogZenServiceInstance, Warning, TEXT("Can't find file for plugin '%s', tried '%s'"), *Name, *AbsPathRoot);
 				return false;
 			}
+
+			Name = FPaths::GetBaseFilename(Name);
 			bHasName = true;
 		}
 		else
@@ -2782,6 +2790,309 @@ PromptUserUnableToUpdateSystemServiceBinaries(const FString& InstallPath)
 		UE_LOG(LogZenServiceInstance, Warning, TEXT("Unreal Zen Storage Server was unable to update the binaries at '%s'.Zen will run using AutoLaunch or ConnectExisting settings."), *InstallPath);
 	}
 }
+struct FInstalledPluginInfo
+{
+	FString Name;
+	uint32_t Version;
+	TArray<FString> References;
+	TMap<FName, FString> Options;
+};
+
+static TArray<FInstalledPluginInfo>
+ReadInstalledPluginInfo()
+{
+	FString PluginInfoContents;
+	TArray<TSharedPtr<FJsonValue>> PluginInfoJson;
+	if (FFileHelper::LoadFileToString(PluginInfoContents, *GetServicePluginsVersionInfoPath()))
+	{
+		TSharedRef<TJsonReader<>> Reader = TJsonReaderFactory<>::Create(PluginInfoContents);
+		if (!FJsonSerializer::Deserialize(Reader, PluginInfoJson))
+		{
+			UE_LOG(LogZenServiceInstance, Warning, TEXT("Found malformed plugin info file at '%s'. Recreating plugin info file."), *GetServicePluginsVersionInfoPath());
+			return {};
+		}
+
+		TArray<FInstalledPluginInfo> PluginInfo;
+		for (const TSharedPtr<FJsonValue>& Value : PluginInfoJson)
+		{
+			int32 Index = PluginInfo.Emplace();
+
+			const TSharedPtr<FJsonObject>& Object = Value->AsObject();
+			PluginInfo[Index].Name = Object->GetStringField(TEXT("name"));
+			PluginInfo[Index].Version = static_cast<uint32_t>(Object->GetNumberField(TEXT("version")));
+
+			for (const TSharedPtr<FJsonValue>& Reference : Object->GetArrayField(TEXT("references")))
+			{
+				PluginInfo[Index].References.Add(Reference->AsString());
+			}
+
+			for (const TSharedPtr<FJsonValue>& Option : Object->GetArrayField(TEXT("options")))
+			{
+				for (const TTuple<FString, TSharedPtr<FJsonValue>>& NameValuePair : Option->AsObject()->Values)
+				{
+					PluginInfo[Index].Options.Add(FName(NameValuePair.Get<0>()), NameValuePair.Get<1>()->AsString());
+				}
+			}
+		}
+
+		return PluginInfo;
+	}
+	else
+	{
+		UE_LOG(LogZenServiceInstance, Display, TEXT("Unable to find existing plugin information at '%s'. Creating a new plugin info file."), *GetServicePluginsVersionInfoPath());
+		return {};
+	}
+}
+
+static bool
+WriteInstalledPluginInfo(const TArray<FInstalledPluginInfo>& Info)
+{
+	TArray<TSharedPtr<FJsonValue>> Json;
+
+	for (const FInstalledPluginInfo& PluginInfo : Info)
+	{
+		TSharedPtr<FJsonObject> Object = MakeShareable(new FJsonObject);
+		Object->SetStringField("name", PluginInfo.Name);
+		Object->SetNumberField("version", PluginInfo.Version);
+
+		TArray<TSharedPtr<FJsonValue>> OptionsJson;
+		for (const TTuple<FName, FString>& Option : PluginInfo.Options)
+		{
+			TSharedPtr<FJsonObject> OptionJson = MakeShareable(new FJsonObject);
+			OptionJson->SetStringField(Option.Get<0>().ToString(), Option.Get<1>());
+			OptionsJson.Add(MakeShareable(new FJsonValueObject(OptionJson)));
+		}
+
+		Object->SetArrayField("options", OptionsJson);
+
+		TArray<TSharedPtr<FJsonValue>> ReferencesJson;
+		for (const FString& Reference : PluginInfo.References)
+		{
+			ReferencesJson.Add(MakeShareable(new FJsonValueString(Reference)));
+		}
+
+		Object->SetArrayField("references", ReferencesJson);
+
+		Json.Add(MakeShareable(new FJsonValueObject(Object)));
+	}
+
+	FString OutputString;
+	TSharedRef<TJsonWriter<>> Writer = TJsonWriterFactory<>::Create(&OutputString);
+	if (!FJsonSerializer::Serialize(Json, Writer))
+	{
+		UE_LOG(LogZenServiceInstance, Warning, TEXT("Failed to serialize plugin version info"));
+		return false;
+	}
+
+	const FString OutputPath = GetServicePluginsVersionInfoPath();
+	UE_LOG(LogZenServiceInstance, Display, TEXT("Writing plugin version info to '%s'"), *OutputPath);
+	return FFileHelper::SaveStringToFile(OutputString, *OutputPath, FFileHelper::EEncodingOptions::ForceUTF8WithoutBOM);
+}
+
+static uint32_t
+GetPluginVersion(const FString& Path)
+{
+#if PLATFORM_WINDOWS
+	HMODULE DllHandle = LoadLibrary(*Path);
+	if (!DllHandle)
+	{
+		UE_LOG(LogZenServiceInstance, Warning, TEXT("Failed to retrieve plugin info from '%s' - unable to load library. This plugin will be skipped."), *Path);
+		return static_cast<uint32_t>(-1);
+	}
+
+	ON_SCOPE_EXIT
+	{
+		FreeLibrary(DllHandle);
+	};
+
+	typedef void (*PfnGetTransportPluginVersion)(uint32_t* OutApiVersion, uint32_t* OutPluginVersion);
+	// Intermediate cast to void to silence compiler warnings about casting a FARPROC to a different
+	// function type - we know the function type is correct here
+	PfnGetTransportPluginVersion GetVersionFn = (PfnGetTransportPluginVersion)(void*)GetProcAddress(DllHandle, "GetTransportPluginVersion");
+
+	if (!GetVersionFn)
+	{
+		UE_LOG(LogZenServiceInstance, Warning, TEXT("Failed to retrieve plugin version info from '%s' - library does not implement GetTransportPluginVersion()"), *Path);
+		return static_cast<uint32_t>(-1);
+	}
+
+	uint32_t ApiVersion;
+	uint32_t PluginVersion;
+
+	GetVersionFn(&ApiVersion, &PluginVersion);
+
+	return PluginVersion;
+#else
+	return static_cast<uint32_t>(-1);
+#endif
+}
+
+// Merges desired plugin info from the set of currently installed plugins and the plugins desired by
+// this workspace's config. InstalledPluginInfo will be updated with any plugins from PluginSettings
+// which are not present or are present at a lower version than the one we require, and PluginSettings
+// will be updated with settings for any plugins present in InstalledPluginInfo. Any plugins that
+// are denylisted with Zen.Plugin.DenyList will be removed from both sets and thus the final
+// installation.
+static bool
+UpdateInstalledPluginInfo(TArray<FInstalledPluginInfo>& InstalledPluginInfo, TArray<FServicePluginSettings>& PluginSettings)
+{
+	bool bNeedsUpdate = false;
+	FString ReferenceDir = FPaths::ConvertRelativePathToFull(FPaths::EngineDir());
+
+	for (FServicePluginSettings& Plugin : PluginSettings)
+	{
+		if (FInstalledPluginInfo* InstalledPlugin = InstalledPluginInfo.FindByPredicate([&Plugin](const FInstalledPluginInfo& Info) { return Info.Name == Plugin.Name; }))
+		{
+			uint32_t Version = GetPluginVersion(Plugin.AbsPath);
+			if (Version == static_cast<uint32_t>(-1))
+			{
+				continue; // skip plugin - we've already logged a warning in GetPluginVersion()
+			}
+
+			if (Version > InstalledPlugin->Version)
+			{
+				UE_LOG(LogZenServiceInstance, Display, TEXT("Updating plugin '%s' from version '%d' to version '%d'"), *Plugin.Name, InstalledPlugin->Version, Version);
+
+				InstalledPlugin->Version = Version;
+				InstalledPlugin->Options = Plugin.Options;
+
+				Plugin.bNeedsInstall = true;
+				bNeedsUpdate = true;
+			}
+
+			for (const TTuple<FName, FString>& Option : Plugin.Options)
+			{
+				const FString* ExistingOption = InstalledPlugin->Options.Find(Option.Get<0>());
+				if (!ExistingOption || *ExistingOption != Option.Get<1>())
+				{
+					UE_LOG(LogZenServiceInstance, Display, TEXT("Updating plugin installation due to new or changed value for option '%s' : '%s'"), *(Option.Get<0>().ToString()), *Option.Get<1>());
+
+					InstalledPlugin->Options.Add(Option.Get<0>(), Option.Get<1>());
+
+					Plugin.bNeedsInstall = true;
+					bNeedsUpdate = true;
+				}
+			}
+
+			if (!InstalledPlugin->References.Contains(ReferenceDir))
+			{
+				InstalledPlugin->References.Add(ReferenceDir);
+			}
+		}
+		else
+		{
+			FInstalledPluginInfo NewPlugin;
+			uint32_t Version = GetPluginVersion(Plugin.AbsPath);
+			if (Version == static_cast<uint32_t>(-1))
+			{
+				continue; // skip plugin - we've already logged a warning in GetPluginVersion()
+			}
+
+			NewPlugin.Name = Plugin.Name;
+			NewPlugin.Version = Version;
+			NewPlugin.Options = Plugin.Options;
+			NewPlugin.References.Add(ReferenceDir);
+
+			InstalledPluginInfo.Add(NewPlugin);
+
+			Plugin.bNeedsInstall = true;
+			bNeedsUpdate = true;
+		}
+	}
+
+	TArray<FString> PluginsToRemove;
+
+	// Check plugins that are installed, but not by this workspace. If we find one,
+	// we remove this workspace from the workspaces referencing it. Then, if there
+	// are no remaining references, we remove it - otherwise, we add it to the desired
+	// plugin settings for this installation, so when we later write the plugin config
+	// it's included.
+	for (FInstalledPluginInfo& Info : InstalledPluginInfo)
+	{
+		const FServicePluginSettings* Plugin = PluginSettings.FindByPredicate([&Info](const FServicePluginSettings& P) { return P.Name == Info.Name; });
+		if (!Plugin)
+		{
+			Info.References.Remove(ReferenceDir);
+			if (Info.References.IsEmpty())
+			{
+				UE_LOG(LogZenServiceInstance, Display, TEXT("Removing installed plugin '%s' as it is no longer referenced by any workspace"), *Info.Name);
+
+				PluginsToRemove.Add(Info.Name);
+
+				bNeedsUpdate = true;
+			}
+			else
+			{
+				FServicePluginSettings NewPlugin;
+				NewPlugin.Name = Info.Name;
+				NewPlugin.Options = Info.Options;
+
+				PluginSettings.Add(NewPlugin);
+			}
+		}
+	}
+
+	// Check known bad versions of plugins and remove them forcibly if found.
+	// Plugins can be denied using [Zen.Plugin.DenyList] in config, where values
+	// in the section are an array of versions, either single versions specified
+	// by +PluginName=X, or potentially open ranges of versions specified by
+	// +PluginName=X-Y (for bounded ranges) or +PluginName=X- (for all versions
+	// above a given version).
+	const FConfigSection* DenyListSection = GConfig->GetSection(TEXT("Zen.Plugin.DenyList"), false, GEngineIni);
+	if (DenyListSection)
+	{
+		for (const FInstalledPluginInfo& Info : InstalledPluginInfo)
+		{
+			TArray<FString> DenyVersions;
+			if (DenyListSection->GetArray(*Info.Name, DenyVersions))
+			{
+				for (const FString& Version : DenyVersions)
+				{
+					uint32 Begin, End;
+					if (!LexTryParseString(Begin, *Version))
+					{
+						UE_LOG(LogZenServiceInstance, Warning, TEXT("Unable to parse plugin deny list version '%s' for plugin '%s'"), *Version, *Info.Name);
+						continue;
+					}
+
+					if (Version.Contains(TEXT("-")))
+					{
+						if (Version.EndsWith(TEXT("-")))
+						{
+							End = INT_MAX;
+						}
+						else
+						{
+							int32 DashIndex = Version.Find(TEXT("-"));
+							if (!LexTryParseString(End, *Version + DashIndex + 1))
+							{
+								UE_LOG(LogZenServiceInstance, Warning, TEXT("Unable to parse plugin deny list version '%s' for plugin '%s'"), *Version, *Info.Name);
+								continue;
+							}
+						}
+					}
+					else
+					{
+						End = Begin;
+					}
+
+					if (Info.Version >= Begin && Info.Version <= End)
+					{
+						UE_LOG(LogZenServiceInstance, Display, TEXT("Removing blacklisted plugin version '%d'"), Info.Version);
+						PluginsToRemove.Add(Info.Name);
+
+						bNeedsUpdate = true;
+					}
+				}
+			}
+		}
+	}
+
+	InstalledPluginInfo.RemoveAll([&](const FInstalledPluginInfo& I) { return PluginsToRemove.Contains(I.Name); });
+	PluginSettings.RemoveAll([&](const FServicePluginSettings& S) { return PluginsToRemove.Contains(S.Name); });
+	
+	return bNeedsUpdate;
+}
 
 static bool IsDirectoryWriteable(const FString& Directory)
 {
@@ -3109,17 +3420,13 @@ bool FZenServiceInstance::ConditionalUpdateSystemServiceInstall(const FSystemSer
 bool
 FZenServiceInstance::ConditionalUpdateLocalInstall(FServiceAutoLaunchSettings::EInstallMode InstallMode)
 {
-	// Normally plugins config is written only if we update zen local install,
-	// extra command line provided to force plugins update.
-	const bool bForceZenPluginsInstall = FParse::Param(FCommandLine::Get(), TEXT("ForceZenPluginsInstall"));
-	// If current settings require plugins, we will pass "--plugins-config" command line argument to zen server,
-	// at that point plugins config file should exist, so create one if it doesn't.
-	const bool bMissingZenPluginsConfig = Settings.IsRequirePlugins() && !IFileManager::Get().FileExists(*GetServicePluginsConfigPath());
-	if (bForceZenPluginsInstall || bMissingZenPluginsConfig)
-	{
-		WriteLocalPluginsConfig();
-		CleanOutOfDateServicePluginConfigs();
-	}
+	bool bNeedsPluginUpdate = FParse::Param(FCommandLine::Get(), TEXT("ForceZenPluginsInstall"));
+
+	// Always update the installed plugin info, as we may need to add a reference from this workspace
+	// even if we don't actually need to perform a full plugin update
+	TArray<FInstalledPluginInfo> InstalledPluginInfo = ReadInstalledPluginInfo();
+	bNeedsPluginUpdate = UpdateInstalledPluginInfo(InstalledPluginInfo, Settings.SettingsVariant.Get<FServiceAutoLaunchSettings>().Plugins) || bNeedsPluginUpdate;
+	WriteInstalledPluginInfo(InstalledPluginInfo);
 
 	if (InstallMode == FServiceAutoLaunchSettings::EInstallMode::Link)
 	{
@@ -3151,12 +3458,24 @@ FZenServiceInstance::ConditionalUpdateLocalInstall(FServiceAutoLaunchSettings::E
 		FString InTreeServicePath = GetInTreeServicePath();
 		if (LinkIsValid && InTreeServicePath == Link.ServicePath)
 		{
-			// If the running process already points to our executable and we have a valid link file we are good to go
+			// If the running process already points to our executable and we have a valid link file we are good to go,
+			// unless a plugin update is required - in that case we don't need to do any shutdown/install of the service
+			// itself, just write the updated plugin config and restart
 			uint32_t Pid = 0;
 			if (ZenServerState::FindRunningProcessId(*InTreeServicePath, &Pid))
 			{
-				UE_LOG(LogZenServiceInstance, Log, TEXT("Service link '%s' pointing to '%s', version %s is up and running"), *LinkPath, *Link.ServicePath, *Link.Version.ToString(false));
-				return true;
+				if (bNeedsPluginUpdate)
+				{
+					UE_LOG(LogZenServiceInstance, Log, TEXT("Plugin update required - stopping running service at '%s'"), *Link.ServicePath);
+					ShutdownZenServerProcess(Pid);
+					WriteLocalPluginsConfig();
+					return true;
+				}
+				else
+				{
+					UE_LOG(LogZenServiceInstance, Log, TEXT("Service link '%s' pointing to '%s', version %s is up and running"), *LinkPath, *Link.ServicePath, *Link.Version.ToString(false));
+					return true;
+				}
 			}
 		}
 
@@ -3179,8 +3498,18 @@ FZenServiceInstance::ConditionalUpdateLocalInstall(FServiceAutoLaunchSettings::E
 				uint32_t Pid = 0;
 				if (ZenServerState::FindRunningProcessId(*Link.ServicePath, &Pid))
 				{
-					UE_LOG(LogZenServiceInstance, Display, TEXT("Service link '%s' pointing to '%s', version %s is up to date and running"), *LinkPath, *Link.ServicePath, *Link.Version.ToString(false));
-					return true;
+					if (bNeedsPluginUpdate)
+					{
+						UE_LOG(LogZenServiceInstance, Log, TEXT("Plugin update required - stopping running service at '%s'"), *Link.ServicePath);
+						ShutdownZenServerProcess(Pid);
+						WriteLocalPluginsConfig();
+						return true;
+					}
+					else
+					{
+						UE_LOG(LogZenServiceInstance, Display, TEXT("Service link '%s' pointing to '%s', version %s is up to date and running"), *LinkPath, *Link.ServicePath, *Link.Version.ToString(false));
+						return true;
+					}
 				}
 
 				// Verify that the executable pointed at is runnable and is of matching version
@@ -3189,6 +3518,7 @@ FZenServiceInstance::ConditionalUpdateLocalInstall(FServiceAutoLaunchSettings::E
 				{
 					if (LinkedVersion == Link.Version)
 					{
+						WriteLocalPluginsConfig();
 						UE_LOG(LogZenServiceInstance, Display, TEXT("Service link '%s' pointing to '%s', version %s is up to date"), *LinkPath, *Link.ServicePath, *Link.Version.ToString(false));
 						return true;
 					}
@@ -3211,14 +3541,14 @@ FZenServiceInstance::ConditionalUpdateLocalInstall(FServiceAutoLaunchSettings::E
 			return false;
 		}
 
-		WriteLocalPluginsConfig();
-
 		FString ServiceCopyInstallPath = GetServiceCopyInstallPath();
 		if (FileManager.FileExists(*ServiceCopyInstallPath))
 		{
 			ShutDownZenServerProcessExecutable(ServiceCopyInstallPath);
 		}
 
+		WriteLocalPluginsConfig();
+
 		TArray<FString> FilesToCleanUp{ GetUtilityCopyInstallPath(), ServiceCopyInstallPath, GetInstallVersionCachePath(), GetInstallVersionCachePath() };
 #if PLATFORM_WINDOWS
 		FilesToCleanUp.Add(FPaths::SetExtension(FilesToCleanUp[0], TEXT("pdb")));
@@ -3286,6 +3616,20 @@ FZenServiceInstance::ConditionalUpdateLocalInstall(FServiceAutoLaunchSettings::E
 			UE_LOG(LogZenServiceInstance, Display, TEXT("Installed service at '%s' is up to date"), *InstallServicePath);
 		}
 
+		// If we need a plugin update but not a full Zen update, we just need to shut
+		// down the instance and update the plugins without performing a full install
+		if (bNeedsPluginUpdate && !bMainExecutablesUpdated)
+		{
+			if (!ShutDownZenServerProcessExecutable(InstallServicePath))
+			{
+				PromptUserToStopRunningServerInstanceForUpdate(InstallServicePath);
+				return false;
+			}
+
+			WriteLocalPluginsConfig();
+			return true;
+		}
+
 #if PLATFORM_WINDOWS
 		struct FZenExecutable
 		{
@@ -3339,11 +3683,7 @@ FZenServiceInstance::ConditionalUpdateLocalInstall(FServiceAutoLaunchSettings::E
 bool
 FZenServiceInstance::WriteLocalPluginsConfig()
 {
-	if (!Settings.IsRequirePlugins())
-	{
-		return false;
-	}
-
+#if PLATFORM_WINDOWS
 	TConstArrayView<FServicePluginSettings> PluginsSettings = Settings.SettingsVariant.Get<FServiceAutoLaunchSettings>().Plugins;
 
 	TArray<TSharedPtr<FJsonValue>> JsonArray;
@@ -3351,9 +3691,20 @@ FZenServiceInstance::WriteLocalPluginsConfig()
 	// replace this with UStructToJsonObject or similar in the future
 	for (const FServicePluginSettings& PluginSettings : PluginsSettings)
 	{
+		FString InstallPath = FPaths::Combine(GetLocalInstallPath(), FPaths::GetCleanFilename(PluginSettings.AbsPath));
+
+		if (PluginSettings.bNeedsInstall)
+		{
+			if (!AttemptFileCopyWithRetries(*InstallPath, *PluginSettings.AbsPath, 5.0))
+			{
+				UE_LOG(LogZenServiceInstance, Warning, TEXT("Failed to copy plugin '%s' to '%s' - plugin will be skipped"), *PluginSettings.AbsPath, *InstallPath);
+				continue;
+			}
+		}
+
 		TSharedPtr<FJsonObject> JsonObject = MakeShareable(new FJsonObject);
 
-		JsonObject->SetStringField(TEXT("name"), PluginSettings.AbsPath); // zen server can treat name as absolute path
+		JsonObject->SetStringField(TEXT("name"), InstallPath); // zen server can treat name as absolute path
 		for (const TPair<FName, FString>& Option: PluginSettings.Options)
 		{
 			JsonObject->SetStringField(Option.Key.ToString(), Option.Value);
@@ -3373,6 +3724,9 @@ FZenServiceInstance::WriteLocalPluginsConfig()
 	const FString OutputPath = GetServicePluginsConfigPath();
 	UE_LOG(LogZenServiceInstance, Display, TEXT("Writing plugin configuration to '%s'"), *OutputPath);
 	return FFileHelper::SaveStringToFile(OutputString, *OutputPath, FFileHelper::EEncodingOptions::ForceUTF8WithoutBOM);
+#else
+	return false;
+#endif
 }
 
 bool
diff --git a/Engine/Source/Developer/Zen/Public/Experimental/ZenServerInterface.h b/Engine/Source/Developer/Zen/Public/Experimental/ZenServerInterface.h
--- a/Engine/Source/Developer/Zen/Public/Experimental/ZenServerInterface.h
+++ b/Engine/Source/Developer/Zen/Public/Experimental/ZenServerInterface.h
@@ -41,6 +41,7 @@ struct FServicePluginSettings
 	FString Name;
 	FString AbsPath;
 	TMap<FName, FString> Options;
+	bool bNeedsInstall = false;
 
 	UE_API bool ReadFromConfig(const FString& InPluginName);
 	UE_API bool ReadFromCompactBinary(FCbFieldView Field);
