BuildStorageTool: Implement cross-bucket search when listing builds.

diff --git a/Engine/Source/Developer/Zen/Internal/ZenBuildUtils.h b/Engine/Source/Developer/Zen/Internal/ZenBuildUtils.h
--- a/Engine/Source/Developer/Zen/Internal/ZenBuildUtils.h
+++ b/Engine/Source/Developer/Zen/Internal/ZenBuildUtils.h
@@ -6,17 +6,9 @@
 
 namespace UE::Zen::Build
 {
-	struct FNamespacePlatformBucketTuple
-	{
-		FString Namespace;
-		FString Platform;
-		FString Bucket;
-	};
-
 	struct FBuildState
 	{
 		FString Namespace;
-		FString Platform;
 		TArray<UE::Zen::Build::FBuildServiceInstance::FBuildRecord> Results;
 	};
 
diff --git a/Engine/Source/Developer/Zen/Private/BuildListRetriever.cpp b/Engine/Source/Developer/Zen/Private/BuildListRetriever.cpp
--- a/Engine/Source/Developer/Zen/Private/BuildListRetriever.cpp
+++ b/Engine/Source/Developer/Zen/Private/BuildListRetriever.cpp
@@ -66,11 +66,6 @@ void FBuildListRetriever::IssueBuildQueries()
 	check(IsInGameThread());
 
 	TMultiMap<FString, FString> NamespacesAndBuckets = ServiceInstance->GetNamespacesAndBuckets();
-	TMultiMap<FString, FString> BucketsToNamespaces;
-	for (const TPair<FString, FString>& NamespaceAndBucket : NamespacesAndBuckets)
-	{
-		BucketsToNamespaces.AddUnique(NamespaceAndBucket.Value, NamespaceAndBucket.Key);
-	}
 	
 	auto StringToSegmentViews = [](const FString& Str, TArray<FStringView>& OutViews)
 		{
@@ -90,87 +85,51 @@ void FBuildListRetriever::IssueBuildQueries()
 			}
 		};
 
-	TMultiMap<FStringView, TArray<FStringView>> NamespacesToBucketSegmentViews;
-	for (const TPair<FString, FString>& NamespaceAndBucket : NamespacesAndBuckets)
-	{
-		BucketsToNamespaces.AddUnique(NamespaceAndBucket.Value, NamespaceAndBucket.Key);
-		TArray<FStringView>& BucketSegmentViews = NamespacesToBucketSegmentViews.Add(NamespaceAndBucket.Key);
-		StringToSegmentViews(NamespaceAndBucket.Value, BucketSegmentViews);
-	}
-
-	const uint32 SegmentIndexProject = 0;
-	const uint32 SegmentIndexBuildType = 1;
-	const uint32 SegmentIndexStream = 2;
-	const uint32 SegmentIndexPlatform = 3;
-	const uint32 SegmentIndexNum = 4;
+	FString BucketRegex = FString::Printf(TEXT("%s\\.%s\\.%s\\..*"), *ProjectName, *BuildType, *Stream);
+	FString BucketPrefix = FString::Printf(TEXT("%s.%s.%s."), *ProjectName, *BuildType, *Stream);
+	TArray<FString> Namespaces;
+	TArray<FString> NamespacesToQuery;
+	NamespacesAndBuckets.GetKeys(Namespaces);
 
-	// Platform list generation
-	TArray<FString> Platforms;
-	for (const TPair<FStringView, TArray<FStringView>>& NamespaceToBucketSegmentViews : NamespacesToBucketSegmentViews)
+	for (const FString& Namespace : Namespaces)
 	{
-		const TArray<FStringView>& BucketSegmentViews = NamespaceToBucketSegmentViews.Value;
-		if (BucketSegmentViews.Num() == SegmentIndexNum)
+		TArray<FString> BucketsInNamespace;
+		NamespacesAndBuckets.MultiFind(Namespace, BucketsInNamespace);
+		for (const FString& Bucket : BucketsInNamespace)
 		{
-			if (BucketSegmentViews[SegmentIndexStream] != Stream)
+			if (Bucket.StartsWith(BucketPrefix))
 			{
-				continue;
+				NamespacesToQuery.Add(Namespace);
+				break;
 			}
-
-			if (BucketSegmentViews[SegmentIndexProject] != ProjectName)
-			{
-				continue;
-			}
-
-			if (BucketSegmentViews[SegmentIndexBuildType] != BuildType)
-			{
-				continue;
-			}
-
-			Platforms.AddUnique(FString(BucketSegmentViews[SegmentIndexPlatform]));
-		}
-	}
-
-	TArray<FNamespacePlatformBucketTuple> NamespacePlatformBucketTuples;
-	for (const FString& Platform : Platforms)
-	{
-		FString Bucket = FString::Printf(TEXT("%s.%s.%s.%s"), *ProjectName, *BuildType, *Stream, *Platform);
-		// TODO: replace with passed in namespace
-		TArray<FString> NamespacesForBucket;
-		BucketsToNamespaces.MultiFind(Bucket, NamespacesForBucket);
-		for (FString& Namespace : NamespacesForBucket)
-		{
-			NamespacePlatformBucketTuples.Emplace(MoveTemp(Namespace), *Platform, Bucket);
 		}
 	}
 
-	if (NamespacePlatformBucketTuples.Num() > 0)
+	if (NamespacesToQuery.Num() > 0)
 	{
 		ReadyToIssueBuildQuery = false;
 		
 		TSharedPtr<FListBuildsState> PendingQueryState = MakeShared<FListBuildsState>();
-		PendingQueryState->PendingQueries = NamespacePlatformBucketTuples.Num();
-		PendingQueryState->QueryState.SetNum(NamespacePlatformBucketTuples.Num());
+		PendingQueryState->PendingQueries = NamespacesToQuery.Num();
+		PendingQueryState->QueryState.SetNum(NamespacesToQuery.Num());
 		uint32 QueryIndex = 0;
 
-		for (const FNamespacePlatformBucketTuple& NamespacePlatformBucket : NamespacePlatformBucketTuples)
+		for (const FString& Namespace : NamespacesToQuery)
 		{
 			FCbWriter QueryWriter;
 			QueryWriter.BeginObject();
-			QueryWriter.BeginObject("query");
 			QueryWriter.BeginObject("isPreflight");
 			QueryWriter.AddString("$eq", "false");
 			QueryWriter.EndObject();// isPreflight
 			QueryWriter.EndObject();// query
-			QueryWriter.EndObject();
 
-			ServiceInstance->ListBuilds(NamespacePlatformBucket.Namespace, NamespacePlatformBucket.Bucket, QueryWriter.Save().AsObject(),
-				[this, QueryIndex, Namespace = NamespacePlatformBucket.Namespace, Platform = NamespacePlatformBucket.Platform,
+			ServiceInstance->ListBuildsAcrossBuckets(Namespace, BucketRegex, QueryWriter.Save().AsObject(),
+				[this, QueryIndex, Namespace = FString(Namespace),
 				PendingQueryState]
 				(TArray<FBuildServiceInstance::FBuildRecord>&& Results) mutable
 				{
 					FBuildState& NewBuildState = PendingQueryState->QueryState[QueryIndex];
 					NewBuildState.Namespace = MoveTemp(Namespace);
-					NewBuildState.Platform = MoveTemp(Platform);
 					NewBuildState.Results = MoveTemp(Results);
 
 					if (--PendingQueryState->PendingQueries == 0)
@@ -202,7 +161,7 @@ void FBuildListRetriever::ExtractBuildsFromQueryResult(FListBuildsState& ListBui
 	{
 		for (FBuildServiceInstance::FBuildRecord& BuildRecord : BuildState.Results)
 		{
-			TArray<int32>& Builds = PerPlatformBuilds.FindOrAdd(BuildState.Platform);
+			TArray<int32>& Builds = PerPlatformBuilds.FindOrAdd(BuildRecord.GetCookPlatform());
 			FString CommitIdentifier = BuildRecord.GetCommitIdentifier();
 			Builds.Add(FCString::Atoi(*CommitIdentifier));
 		}
diff --git a/Engine/Source/Developer/Zen/Private/BuildServerInterface.cpp b/Engine/Source/Developer/Zen/Private/BuildServerInterface.cpp
--- a/Engine/Source/Developer/Zen/Private/BuildServerInterface.cpp
+++ b/Engine/Source/Developer/Zen/Private/BuildServerInterface.cpp
@@ -397,7 +397,9 @@ FBuildServiceInstance::ListBuilds(FStringView Namespace, FStringView Bucket, FCb
 	{
 		TUniquePtr<FArchive> QueryFileArchive(IFileManager::Get().CreateFileWriter(*QueryFilePath, FILEWRITE_NoFail));
 		FCbWriter Writer;
+		Writer.BeginObject();
 		Writer.AddObject("query", Query);
+		Writer.EndObject();
 		Writer.Save(*QueryFileArchive);
 	};
 
@@ -448,6 +450,90 @@ FBuildServiceInstance::ListBuilds(FStringView Namespace, FStringView Bucket, FCb
 	});
 }
 
+void
+FBuildServiceInstance::ListBuildsAcrossBuckets(FStringView Namespace, FStringView BucketRegex, FCbObject&& Query, FOnListBuildsComplete&& InOnListBuildsComplete)
+{
+	if (ConnectionState.load(std::memory_order_relaxed) != EConnectionState::ConnectionSucceeded)
+	{
+		if (InOnListBuildsComplete)
+		{
+			InOnListBuildsComplete({});
+		}
+		return;
+	}
+
+	FString QueryFilePath = FPaths::CreateTempFilename(FPlatformProcess::UserTempDir(), TEXT("BuildListQuery"), TEXT(".cbo"));
+	FPaths::MakePlatformFilename(QueryFilePath);
+
+	FString OutputFilePath = FPaths::CreateTempFilename(FPlatformProcess::UserTempDir(), TEXT("BuildList"), TEXT(".cbo"));
+	FPaths::MakePlatformFilename(OutputFilePath);
+
+	FString CommandLineArgs = FString::Printf(TEXT("builds list --namespace \"%s\" \"%s\" \"%s\""),
+		*WriteToString<64>(Namespace), *QueryFilePath, *OutputFilePath);
+	CommandLineArgs = AddZenUtilityBuildServerArguments(*CommandLineArgs);
+
+	auto PreInvoke = [QueryFilePath, BucketRegex = FString(BucketRegex), Query = MoveTemp(Query)](FString& CommandLineArgs)
+	{
+		TUniquePtr<FArchive> QueryFileArchive(IFileManager::Get().CreateFileWriter(*QueryFilePath, FILEWRITE_NoFail));
+		FCbWriter Writer;
+		Writer.BeginObject();
+		Writer.AddString("bucketRegex", BucketRegex);
+		Writer.AddObject("query", Query);
+		Writer.BeginObject("options");
+		Writer.AddInteger("limit", 5000);
+		Writer.AddInteger("max", 50000);
+		Writer.EndObject();
+		Writer.EndObject();
+		Writer.Save(*QueryFileArchive);
+	};
+
+	InvokeZenUtility(CommandLineArgs, MoveTemp(PreInvoke),
+	[QueryFilePath = MoveTemp(QueryFilePath), OutputFilePath = MoveTemp(OutputFilePath), InOnListBuildsComplete = MoveTemp(InOnListBuildsComplete)](bool bSuccessful) mutable
+	{
+		ON_SCOPE_EXIT
+		{
+			IFileManager & FileManager = IFileManager::Get();
+			FileManager.Delete(*QueryFilePath, false, false, true);
+			FileManager.Delete(*OutputFilePath, false, false, true);
+		};
+
+		if (!bSuccessful)
+		{
+			if (InOnListBuildsComplete)
+			{
+				InOnListBuildsComplete({});
+			}
+			return;
+		}
+
+		TUniquePtr<FArchive> Ar(IFileManager::Get().CreateFileReader(*OutputFilePath));
+		if (!Ar)
+		{
+			UE_LOG(LogBuildServiceInstance, Warning, TEXT("Missing output file from zen utility when gathering build data: '%s'"), *OutputFilePath);
+			if (InOnListBuildsComplete)
+			{
+				InOnListBuildsComplete({});
+			}
+			return;
+		}
+
+		if (InOnListBuildsComplete)
+		{
+			FCbObject OutputObject = LoadCompactBinary(*Ar).AsObject();
+			TArray<FBuildRecord> BuildRecords;
+			for (FCbFieldView ResultField : OutputObject["results"].AsArrayView())
+			{
+				FBuildRecord BuildRecord;
+				if (LoadFromCompactBinary(ResultField, {}, BuildRecord))
+				{
+					BuildRecords.Emplace(MoveTemp(BuildRecord));
+				}
+			}
+			InOnListBuildsComplete(MoveTemp(BuildRecords));
+		}
+	});
+}
+
 FBuildServiceInstance::FBuildTransfer
 FBuildServiceInstance::StartBuildTransfer(
 	const FCbObjectId& BuildId,
@@ -595,6 +681,17 @@ FString FBuildServiceInstance::FBuildRecord::GetCommitIdentifier() const
 	return CommitIdentifier;
 }
 
+FString FBuildServiceInstance::FBuildRecord::GetCookPlatform() const
+{
+	FString CookPlatform;
+	if (FCbFieldView CookPlatformField = Metadata["cookPlatform"]; CookPlatformField.HasValue() && !CookPlatformField.HasError())
+	{
+		CookPlatform = *WriteToString<64>(CookPlatformField.AsString());
+	}
+
+	return CookPlatform;
+}
+
 FString
 FBuildServiceInstance::FBuildTransfer::GetDescription() const
 {
diff --git a/Engine/Source/Developer/Zen/Public/Experimental/BuildServerInterface.h b/Engine/Source/Developer/Zen/Public/Experimental/BuildServerInterface.h
--- a/Engine/Source/Developer/Zen/Public/Experimental/BuildServerInterface.h
+++ b/Engine/Source/Developer/Zen/Public/Experimental/BuildServerInterface.h
@@ -86,6 +86,7 @@ public:
 		FCbObject Metadata;
 
 		UE_API FString GetCommitIdentifier() const;
+		UE_API FString GetCookPlatform() const;
 	};
 
 	struct FBuildTransfer
@@ -137,6 +138,7 @@ public:
 	inline const TMultiMap<FString, FString>& GetNamespacesAndBuckets() const { return NamespacesAndBuckets; }
 
 	UE_API void ListBuilds(FStringView Namespace, FStringView Bucket, FCbObject&& Query, FOnListBuildsComplete&& InOnListBuildsComplete);
+	UE_API void ListBuildsAcrossBuckets(FStringView Namespace, FStringView BucketRegex, FCbObject&& Query, FOnListBuildsComplete&& InOnListBuildsComplete);
 
 	UE_API FBuildTransfer StartBuildTransfer(
 		const FCbObjectId& BuildId,
diff --git a/Engine/Source/Editor/StorageServerWidgets/Private/SBuildSelection.cpp b/Engine/Source/Editor/StorageServerWidgets/Private/SBuildSelection.cpp
--- a/Engine/Source/Editor/StorageServerWidgets/Private/SBuildSelection.cpp
+++ b/Engine/Source/Editor/StorageServerWidgets/Private/SBuildSelection.cpp
@@ -205,20 +205,21 @@ void SBuildSelection::RebuildLists()
 	using namespace UE::BuildSelection::Internal;
 
 	BuildListRefreshesInProgress.fetch_add(1);
-	BucketsToNamespaces.Empty();
 	StreamList.Empty();
 	ProjectList.Empty();
 	BuildTypeList.Empty();
 	PlatformList.Empty();
 
+	TArray<FString> Namespaces;
+	TMultiMap<FString, FString> NamespacesAndBuckets;
+
 	if (TSharedPtr<FBuildServiceInstance> ServiceInstance = BuildServiceInstance.Get())
 	{
-		TArray<FString> Namespaces;
 		TArray<FString> Projects;
 		TArray<FString> Streams;
 		TArray<FString> BuildTypes;
 		TArray<FString> Platforms;
-		TMultiMap<FString, FString> NamespacesAndBuckets = ServiceInstance->GetNamespacesAndBuckets();
+		NamespacesAndBuckets = ServiceInstance->GetNamespacesAndBuckets();
 
 		auto StringToSegmentViews = [](const FString& Str, TArray<FStringView>& OutViews)
 			{
@@ -276,7 +277,6 @@ void SBuildSelection::RebuildLists()
 		TMultiMap<FStringView, TArray<FStringView>> NamespacesToBucketSegmentViews;
 		for (const TPair<FString, FString>& NamespaceAndBucket : NamespacesAndBuckets)
 		{
-			BucketsToNamespaces.AddUnique(NamespaceAndBucket.Value, NamespaceAndBucket.Key);
 			TArray<FStringView>& BucketSegmentViews = NamespacesToBucketSegmentViews.Add(NamespaceAndBucket.Key);
 			StringToSegmentViews(NamespaceAndBucket.Value, BucketSegmentViews);
 
@@ -411,36 +411,42 @@ void SBuildSelection::RebuildLists()
 
 	if (TSharedPtr<FBuildServiceInstance> ServiceInstance = BuildServiceInstance.Get())
 	{
-		TArray<FNamespacePlatformBucketTuple> NamespacePlatformBucketTuples;
-		for (TSharedPtr<FString> Platform : PlatformList)
+		FString BucketRegex = FString::Printf(TEXT("%s\\.%s\\.%s\\..*"), **SelectedProject, **SelectedBuildType, **SelectedStream);
+		FString BucketPrefix = FString::Printf(TEXT("%s.%s.%s."), **SelectedProject, **SelectedBuildType, **SelectedStream);
+
+		TArray<FString> NamespacesToQuery;
+
+		for (const FString& Namespace : Namespaces)
 		{
-			FString Bucket = FString::Printf(TEXT("%s.%s.%s.%s"), **SelectedProject, **SelectedBuildType, **SelectedStream, **Platform);
-			TArray<FString> NamespacesForBucket;
-			BucketsToNamespaces.MultiFind(Bucket, NamespacesForBucket);
-			for (FString& Namespace : NamespacesForBucket)
+			TArray<FString> BucketsInNamespace;
+			NamespacesAndBuckets.MultiFind(Namespace, BucketsInNamespace);
+			for (const FString& Bucket : BucketsInNamespace)
 			{
-				NamespacePlatformBucketTuples.Emplace(MoveTemp(Namespace), *Platform, Bucket);
+				if (Bucket.StartsWith(BucketPrefix))
+				{
+					NamespacesToQuery.Add(Namespace);
+					break;
+				}
 			}
 		}
 
 		TSharedPtr<FListBuildsState> PendingQueryState = MakeShared<FListBuildsState>();
-		PendingQueryState->PendingQueries = NamespacePlatformBucketTuples.Num();
-		PendingQueryState->QueryState.SetNum(NamespacePlatformBucketTuples.Num());
+		PendingQueryState->PendingQueries = NamespacesToQuery.Num();
+		PendingQueryState->QueryState.SetNum(NamespacesToQuery.Num());
 		uint32 QueryIndex = 0;
 
 		++BuildRefreshGeneration;
 		BuildGroups.Empty();
 
-		for (const FNamespacePlatformBucketTuple& NamespacePlatformBucket : NamespacePlatformBucketTuples)
+		for (const FString& Namespace : NamespacesToQuery)
 		{
-			ServiceInstance->ListBuilds(NamespacePlatformBucket.Namespace, NamespacePlatformBucket.Bucket, FCbObject(),
-			[this, QueryIndex, Namespace = NamespacePlatformBucket.Namespace, Platform = NamespacePlatformBucket.Platform,
+			ServiceInstance->ListBuildsAcrossBuckets(Namespace, BucketRegex, FCbObject(),
+			[this, QueryIndex, Namespace = FString(Namespace),
 				ExpectedBuildRefreshGeneration = BuildRefreshGeneration.load(), PendingQueryState]
 			(TArray<FBuildServiceInstance::FBuildRecord>&& Results) mutable
 				{
 					FBuildState& NewBuildState = PendingQueryState->QueryState[QueryIndex];
 					NewBuildState.Namespace = MoveTemp(Namespace);
-					NewBuildState.Platform = MoveTemp(Platform);
 					NewBuildState.Results = MoveTemp(Results);
 
 					if (--PendingQueryState->PendingQueries == 0)
@@ -472,7 +478,7 @@ void SBuildSelection::RebuildLists()
 			++QueryIndex;
 		}
 
-		if (NamespacePlatformBucketTuples.IsEmpty())
+		if (NamespacesToQuery.IsEmpty())
 		{
 			BuildListRefreshesInProgress.fetch_sub(1);
 		}
@@ -567,11 +573,6 @@ void SBuildSelection::RegenerateBuildGroups(UE::Zen::Build::FListBuildsState& Li
 					ArtifactName.Empty();
 				}
 			}
-
-			if (ArtifactName.IsEmpty())
-			{
-				ArtifactName = BuildState.Platform;
-			}
 			
 			if (BuildGroup->DisplayName.IsEmpty())
 			{
@@ -644,6 +645,7 @@ void SBuildSelection::RegenerateBuildGroups(UE::Zen::Build::FListBuildsState& Li
 		}
 	}
 
+	ArtifactNameSet.StableSort(TLess());
 	KeyedGroups.GenerateValueArray(BuildGroups);
 	ConditionalSortBuildGroups();
 }
