Fixed a missing include in ZenServerState.h

#Synced-CL 0

diff --git a/Engine/Source/Developer/Zen/Private/ZenServerState.h b/Engine/Source/Developer/Zen/Private/ZenServerState.h
--- a/Engine/Source/Developer/Zen/Private/ZenServerState.h
+++ b/Engine/Source/Developer/Zen/Private/ZenServerState.h
@@ -2,9 +2,10 @@
 
 #pragma once
 
+#include "Containers/AnsiString.h"
 #include "Containers/StringFwd.h"
+#include "Containers/UnrealString.h"
 #include "Serialization/CompactBinary.h"
-#include "Containers/AnsiString.h"
 
 class ZenServerState
 {
