Set permissions on data directory to allow LocalService user to write data when running Zen in service mode.

[RN] none
#Synced-CL 48600208

diff --git a/Engine/Source/Developer/Zen/Private/ZenServerInterface.cpp b/Engine/Source/Developer/Zen/Private/ZenServerInterface.cpp
--- a/Engine/Source/Developer/Zen/Private/ZenServerInterface.cpp
+++ b/Engine/Source/Developer/Zen/Private/ZenServerInterface.cpp
@@ -41,6 +41,8 @@
 #	include <Windows.h>
 #	include "Windows/HideWindowsPlatformTypes.h"
 #	include <dbghelp.h>
+#	include <accctrl.h>
+#	include <AclAPI.h>
 #endif
 
 #define ALLOW_SETTINGS_OVERRIDE_FROM_COMMANDLINE			(UE_SERVER || !(UE_BUILD_SHIPPING))
@@ -2737,6 +2739,24 @@ PromptUserOfFailedSystemServiceInstall()
 	}
 }
 
+static void
+PromptUserFailedToSetDataPathPermissions()
+{
+#if !IS_PROGRAM
+	if (!FApp::IsUnattended() && !IsRunningCommandlet() && !GIsRunningUnattendedScript)
+	{
+		FText ZenDataPathPermissionsFailurePromptTitle = NSLOCTEXT("Zen", "Zen_DataPathPermissionsFailurePromptTitle", "Failed to install system service");
+		FText ZenDataPathPermissionsFailurePromptText = NSLOCTEXT("Zen", "Zen_DataPathPermissionsFailurePromptText", "Failed to update permissions on existing data path for service user. Zen will run using AutoLaunch or ConnectExisting settings.");
+		FPlatformMisc::MessageBoxExt(EAppMsgType::Ok, *ZenDataPathPermissionsFailurePromptTitle.ToString(), *ZenDataPathPermissionsFailurePromptText.ToString());
+	}
+	else
+#endif
+	{
+		// Just log as there is no one to show a message
+		UE_LOG(LogZenServiceInstance, Warning, TEXT("Failed to update permissions on existing data path for service user. Zen will run using AutoLaunch or ConnectExisting settings."));
+	}
+}
+
 static void
 PromptUserOfFailedSystemServiceUninstall()
 {
@@ -3322,6 +3342,65 @@ bool UninstallSystemService(FString UtilityPath)
 	return true;
 }
 
+static
+bool UpdateSystemServiceDataPathPermissions(const FString& DataPath)
+{
+#if PLATFORM_WINDOWS
+	if (!FPaths::DirectoryExists(DataPath))
+	{
+		// We're fine here - zen will create the data path when we're already running as NT AUTHORITY\LocalService,
+		// so the service account will have the appropriate permissions
+		return true;
+	}
+
+	DWORD Result = 0;
+	PACL OldAcl = nullptr;
+	PACL NewAcl = nullptr;
+	PSECURITY_DESCRIPTOR SecurityDescriptor = nullptr;
+	EXPLICIT_ACCESS Access{};
+
+	Result = GetNamedSecurityInfo(*DataPath, SE_FILE_OBJECT, DACL_SECURITY_INFORMATION, NULL, NULL, &OldAcl, NULL, &SecurityDescriptor);
+	if (Result != ERROR_SUCCESS)
+	{
+		UE_LOG(LogZenServiceInstance, Warning, TEXT("Failed to update data path permissions - GetNamedSecurityInfo failed with %d"), Result);
+		return false;
+	}
+
+	ON_SCOPE_EXIT{ LocalFree(SecurityDescriptor); };
+
+	Access.grfAccessMode = GRANT_ACCESS;
+	Access.grfAccessPermissions = STANDARD_RIGHTS_ALL | FILE_ALL_ACCESS;
+	Access.grfInheritance = SUB_CONTAINERS_AND_OBJECTS_INHERIT;
+	Access.Trustee.TrusteeForm = TRUSTEE_IS_NAME;
+	Access.Trustee.TrusteeType = TRUSTEE_IS_USER;
+	Access.Trustee.MultipleTrusteeOperation = NO_MULTIPLE_TRUSTEE;
+
+	TCHAR AccountName[26];
+	FPlatformString::Strcpy(AccountName, TEXT("NT AUTHORITY\\LocalService"));
+	Access.Trustee.ptstrName = AccountName;
+
+	Result = SetEntriesInAcl(1, &Access, OldAcl, &NewAcl);
+	if (Result != ERROR_SUCCESS)
+	{
+		UE_LOG(LogZenServiceInstance, Warning, TEXT("Failed to update data path permissions - SetEntriesInAcl failed with %d"), Result);
+		return false;
+	}
+
+	ON_SCOPE_EXIT{ LocalFree(NewAcl); };
+
+	TArray<TCHAR> DataPathArr(*DataPath, DataPath.Len() + 1);
+
+	Result = SetNamedSecurityInfo(DataPathArr.GetData(), SE_FILE_OBJECT, DACL_SECURITY_INFORMATION, NULL, NULL, NewAcl, NULL);
+	if (Result != ERROR_SUCCESS)
+	{
+		UE_LOG(LogZenServiceInstance, Warning, TEXT("Failed to update data path permissions - SetNamedSecurityInfo failed with %d"), Result);
+		return false;
+	}
+
+#endif
+	return true;
+}
+
 bool FZenServiceInstance::ConditionalUpdateSystemServiceInstall(const FSystemServiceSettings &InSettings)
 {
 	FString SystemServiceInstallPath = GetSystemServiceInstallPath(InSettings);
@@ -3377,6 +3456,12 @@ bool FZenServiceInstance::ConditionalUpdateSystemServiceInstall(const FSystemSer
 		ServiceInstallCommandLine.Appendf(TEXT(" --user %s"), FPlatformProcess::UserName(false));
 #endif
 
+		if (!UpdateSystemServiceDataPathPermissions(InSettings.DataPath))
+		{
+			PromptUserFailedToSetDataPathPermissions();
+			return false;
+		}
+
 		if (!RunZenUtilityElevated(InTreeUtilityPath, ServiceInstallCommandLine, Output, ReturnCode, DesiredCommandLine))
 		{
 			PromptUserOfFailedSystemServiceInstall();
