adding a richer interface for build list query

#rnx

#Synced-CL 48803560

diff --git a/Engine/Source/Developer/Zen/Private/BuildListRetriever.cpp b/Engine/Source/Developer/Zen/Private/BuildListRetriever.cpp
--- a/Engine/Source/Developer/Zen/Private/BuildListRetriever.cpp
+++ b/Engine/Source/Developer/Zen/Private/BuildListRetriever.cpp
@@ -5,6 +5,9 @@
 #include "Misc/DateTime.h"
 #include "Misc/EngineVersion.h"
 #include "Serialization/CompactBinaryWriter.h"
+#include "Internationalization/Internationalization.h"
+
+#define LOCTEXT_NAMESPACE "FBuildListRetriever"
 
 namespace UE::Zen::Build
 {
@@ -19,7 +22,6 @@ void FBuildListRetriever::ConnectToBuildService()
 	if (bConnected)
 		return;
 
-	bConnected = true;
 	ServiceInstance->OnRefreshNamespacesAndBucketsComplete().AddSP(this, &FBuildListRetriever::OnNamespacesAndBucketsRefreshed);
 	ServiceInstance->Connect(!FApp::IsUnattended(), [this]
 		(UE::Zen::Build::FBuildServiceInstance::EConnectionState ConnectionState,
@@ -27,17 +29,52 @@ void FBuildListRetriever::ConnectToBuildService()
 		{
 			if (ConnectionState == UE::Zen::Build::FBuildServiceInstance::EConnectionState::ConnectionSucceeded)
 			{
+				bConnected = true;
 				ServiceInstance->RefreshNamespacesAndBuckets();
 			}
 		});
 }
 
+bool FBuildListRetriever::IsConnected() const
+{
+	return bConnected;
+}
+
+TSharedRef<UE::Zen::Build::FBuildServiceInstance> FBuildListRetriever::GetServiceInstance() const
+{
+	return ServiceInstance.ToSharedRef();
+}
+
+bool FBuildListRetriever::IsReadyForQuery() const
+{
+	return ReadyToIssueBuildQuery;
+}
+
+void FBuildListRetriever::SetOnReadyForQuery( FOnReadyForQuery&& InOnReadyForQuery )
+{
+	OnReadyForQuery = MoveTemp(InOnReadyForQuery);
+}
+
+void FBuildListRetriever::Refresh()
+{
+	if (bConnected && ReadyToIssueBuildQuery)
+	{
+		ReadyToIssueBuildQuery = false;
+		ServiceInstance->RefreshNamespacesAndBuckets();
+	}
+}
+
 void FBuildListRetriever::OnNamespacesAndBucketsRefreshed()
 {
 	ExecuteOnGameThread(UE_SOURCE_LOCATION,
 		[this]
 		{
 			ReadyToIssueBuildQuery = true;
+
+			if (OnReadyForQuery)
+			{
+				OnReadyForQuery();
+			}
 		});
 }
 
@@ -51,7 +88,27 @@ bool FBuildListRetriever::QueryBuilds(const FString& InProjectName, const FStrin
 		Stream = InStream;
 		OnQueriesCompleteCallback = MoveTemp(InOnQueriesCompleteCallback);
 
-		IssueBuildQueries();
+		IssueBuildQueries(false);
+		return true;
+	}
+	else
+	{
+		// Don't do anything if an existing query is in flight
+		return false;
+	}
+}
+
+bool FBuildListRetriever::QueryBuildGroups(const FString& InProjectName, const FString& InBuildType, const FString& InStream, FOnGroupQueriesComplete&& InOnQueriesCompleteCallback)
+{
+	check(IsInGameThread());
+	if (ReadyToIssueBuildQuery)
+	{
+		ProjectName = InProjectName;
+		BuildType = InBuildType;
+		Stream = InStream;
+		OnGroupQueriesCompleteCallback = MoveTemp(InOnQueriesCompleteCallback);
+
+		IssueBuildQueries(true);
 		return true;
 	}
 	else
@@ -61,30 +118,12 @@ bool FBuildListRetriever::QueryBuilds(const FString& InProjectName, const FStrin
 	}
 }
 
-void FBuildListRetriever::IssueBuildQueries()
+void FBuildListRetriever::IssueBuildQueries(bool bIsGroupQuery)
 {
 	check(IsInGameThread());
 
 	TMultiMap<FString, FString> NamespacesAndBuckets = ServiceInstance->GetNamespacesAndBuckets();
 	
-	auto StringToSegmentViews = [](const FString& Str, TArray<FStringView>& OutViews)
-		{
-			FStringView WorkingStringView(Str);
-			int32 CurrentIndex = 0;
-			while (WorkingStringView.FindChar(TCHAR('.'), CurrentIndex))
-			{
-				if (CurrentIndex != 0)
-				{
-					OutViews.Add(WorkingStringView.Left(CurrentIndex));
-				}
-				WorkingStringView.RightChopInline(CurrentIndex + 1);
-			}
-			if (!WorkingStringView.IsEmpty())
-			{
-				OutViews.Add(WorkingStringView);
-			}
-		};
-
 	FString BucketRegex = FString::Printf(TEXT("%s\\.%s\\.%s\\..*"), *ProjectName, *BuildType, *Stream);
 	FString BucketPrefix = FString::Printf(TEXT("%s.%s.%s."), *ProjectName, *BuildType, *Stream);
 	TArray<FString> Namespaces;
@@ -118,13 +157,16 @@ void FBuildListRetriever::IssueBuildQueries()
 		{
 			FCbWriter QueryWriter;
 			QueryWriter.BeginObject();
-			QueryWriter.BeginObject("isPreflight");
-			QueryWriter.AddString("$eq", "false");
-			QueryWriter.EndObject();// isPreflight
+			if (!bIsGroupQuery)
+			{
+				QueryWriter.BeginObject("isPreflight");
+				QueryWriter.AddString("$eq", "false");
+				QueryWriter.EndObject();// isPreflight
+			}
 			QueryWriter.EndObject();// query
 
 			ServiceInstance->ListBuildsAcrossBuckets(Namespace, BucketRegex, QueryWriter.Save().AsObject(),
-				[this, QueryIndex, Namespace = FString(Namespace),
+				[this, bIsGroupQuery, QueryIndex, Namespace = FString(Namespace),
 				PendingQueryState]
 				(TArray<FBuildServiceInstance::FBuildRecord>&& Results) mutable
 				{
@@ -135,25 +177,81 @@ void FBuildListRetriever::IssueBuildQueries()
 					if (--PendingQueryState->PendingQueries == 0)
 					{
 						// All queries complete
-						ExtractBuildsFromQueryResult(*PendingQueryState);
+						if (bIsGroupQuery)
+						{
+							TArray<TSharedRef<FBuildGroup>> BuildGroups = ExtractBuildGroupsFromQueryResult(*PendingQueryState);
+
+							ExecuteOnGameThread(UE_SOURCE_LOCATION,
+								[this, BuildGroups = MoveTemp(BuildGroups)]() mutable
+								{
+									if (OnGroupQueriesCompleteCallback)
+									{
+										OnGroupQueriesCompleteCallback(BuildGroups);
+										OnGroupQueriesCompleteCallback = nullptr;
+									}
+									ReadyToIssueBuildQuery = true;
+								});
+
+						}
+						else
+						{
+							ExtractBuildsFromQueryResult(*PendingQueryState);
 
-						ExecuteOnGameThread(UE_SOURCE_LOCATION,
-							[this]
-							{
-								if (OnQueriesCompleteCallback)
+							ExecuteOnGameThread(UE_SOURCE_LOCATION,
+								[this]
 								{
-									OnQueriesCompleteCallback(PerPlatformBuilds);
-								}
-								ReadyToIssueBuildQuery = true;
-							});
+									if (OnQueriesCompleteCallback)
+									{
+										OnQueriesCompleteCallback(PerPlatformBuilds);
+										OnQueriesCompleteCallback = nullptr;
+									}
+									ReadyToIssueBuildQuery = true;
+								});
+						}
 					}
 				});
 
 			++QueryIndex;
 		}
 	}
+	else
+	{
+		// no namespaces found - send empty results
+		if (bIsGroupQuery)
+		{
+			ExecuteOnGameThread(UE_SOURCE_LOCATION,
+				[this]() mutable
+				{
+					if (OnGroupQueriesCompleteCallback)
+					{
+						TArray<TSharedRef<FBuildGroup>> BuildGroups;
+						OnGroupQueriesCompleteCallback(BuildGroups);
+						OnGroupQueriesCompleteCallback = nullptr;
+					}
+					ReadyToIssueBuildQuery = true;
+				});
+
+		}
+		else
+		{
+			PerPlatformBuilds.Reset();
+
+			ExecuteOnGameThread(UE_SOURCE_LOCATION,
+				[this]
+				{
+					if (OnQueriesCompleteCallback)
+					{
+						OnQueriesCompleteCallback(PerPlatformBuilds);
+						OnQueriesCompleteCallback = nullptr;
+					}
+					ReadyToIssueBuildQuery = true;
+				});
+		}
+	}
 }
 
+
+
 void FBuildListRetriever::ExtractBuildsFromQueryResult(FListBuildsState& ListBuildsState)
 {
 	PerPlatformBuilds.Reset();
@@ -174,4 +272,295 @@ void FBuildListRetriever::ExtractBuildsFromQueryResult(FListBuildsState& ListBui
 	}
 }
 
+// adapted from SBuildSelection::RegenerateBuildGroups
+TArray<TSharedRef<FBuildListRetriever::FBuildGroup>> FBuildListRetriever::ExtractBuildGroupsFromQueryResult(FListBuildsState& ListBuildsState)
+{
+	TArray<TSharedRef<FBuildGroup>> Result;
+
+	auto MakeGroupKey = [](const FString& Namespace, const FString& CommitIdentifier, const FBuildServiceInstance::FBuildRecord& BuildRecord)
+	{
+		if (FCbFieldView BuildGroupField = BuildRecord.Metadata.FindIgnoreCase("buildgroup"); BuildGroupField.HasValue() && !BuildGroupField.HasError())
+		{
+			return FString(WriteToString<64>(Namespace, ".", BuildGroupField.AsString()));
+		}
+		if (FCbFieldView NameField = BuildRecord.Metadata.FindIgnoreCase("name"); NameField.HasValue() && !NameField.HasError())
+		{
+			FString CandidateKeyName = FString(WriteToString<64>(Namespace, ".", NameField.AsString()));
+			// TODO: This name manipulation needs to be removed when the metadata is more consistent.
+			int32 CLIndex = CandidateKeyName.Find(TEXT("-CL-"), ESearchCase::IgnoreCase, ESearchDir::FromEnd);
+			if (CLIndex != INDEX_NONE)
+			{
+				int32 DashIndex = CandidateKeyName.Find(TEXT("-"), ESearchCase::IgnoreCase, ESearchDir::FromStart, CLIndex + 4);
+				if (DashIndex != INDEX_NONE)
+				{
+					CandidateKeyName.LeftInline(DashIndex);
+				}
+				else
+				{
+					int32 DotIndex = CandidateKeyName.Find(TEXT("."), ESearchCase::IgnoreCase, ESearchDir::FromStart, CLIndex + 4);
+					if (DotIndex != INDEX_NONE)
+					{
+						CandidateKeyName.LeftInline(DotIndex);
+					}
+				}
+			}
+			return CandidateKeyName;
+		}
+		if (!CommitIdentifier.IsEmpty())
+		{
+			return FString(WriteToString<64>(Namespace, ".", CommitIdentifier));
+		}
+		return FString(WriteToString<64>(BuildRecord.BuildId));
+	};
+	TMap<FString, TSharedPtr<FBuildGroup>> KeyedGroups;
+
+
+
+
+	for (FBuildState& BuildState : ListBuildsState.QueryState)
+	{
+		for (FBuildServiceInstance::FBuildRecord& BuildRecord : BuildState.Results)
+		{
+			FString CommitIdentifier = BuildRecord.GetCommitIdentifier();
+			int32 Changelist = FCString::Atoi(*CommitIdentifier);
+
+			FString GroupKey = MakeGroupKey(BuildState.Namespace, CommitIdentifier, BuildRecord);
+			TSharedPtr<FBuildGroup>& BuildGroup = KeyedGroups.FindOrAdd(GroupKey);
+			if (!BuildGroup)
+			{
+				BuildGroup = MakeShared<FBuildGroup>();
+			}
+
+			FString ArtifactName;
+			if (FCbFieldView NameView = BuildRecord.Metadata.FindIgnoreCase("name"); NameView.HasValue() && !NameView.HasError())
+			{
+				ArtifactName = FUTF8ToTCHAR(NameView.AsString());
+
+				// TODO: This name manipulation needs to be removed when the metadata is more consistent.
+				int32 TruncationIndex;
+				int32 CLIndex = ArtifactName.Find(TEXT("-CL-"), ESearchCase::IgnoreCase, ESearchDir::FromEnd);
+				if (CLIndex != INDEX_NONE)
+				{
+					ArtifactName.RightChopInline(CLIndex + 4);
+					if (ArtifactName.FindChar(TCHAR('-'), TruncationIndex))
+					{
+						ArtifactName.RightChopInline(TruncationIndex + 1);
+					}
+				}
+				if (ArtifactName.FindLastChar(TCHAR('.'), TruncationIndex))
+				{
+					ArtifactName.RightChopInline(TruncationIndex + 1);
+				}
+			}
+
+			if (BuildGroup->DisplayName.IsEmpty())
+			{
+				FString Job;
+				if (FCbFieldView JobField = BuildRecord.Metadata.FindIgnoreCase("job"); JobField.HasValue() && !JobField.HasError())
+				{
+					Job = FUTF8ToTCHAR(JobField.AsString());
+				}
+
+				FString Category;
+				if (FCbFieldView TemplateIdField = BuildRecord.Metadata.FindIgnoreCase("hordeTemplateId"); TemplateIdField.HasValue() && !TemplateIdField.HasError())
+				{
+					Category = FUTF8ToTCHAR(TemplateIdField.AsString());
+				}
+
+				bool bIsPreflight = false;
+				if (FCbFieldView IsPreflightField = BuildRecord.Metadata.FindIgnoreCase("ispreflight"); IsPreflightField.HasValue() && !IsPreflightField.HasError())
+				{
+					FString IsPreflightString;
+					IsPreflightString = FUTF8ToTCHAR(IsPreflightField.AsString());
+					bIsPreflight = IsPreflightString.ToBool();
+				}
+
+				FDateTime CreatedAt;
+				if (FCbFieldView CreatedAtField = BuildRecord.Metadata.FindIgnoreCase("createdAt"); CreatedAtField.HasValue() && !CreatedAtField.HasError())
+				{
+					if (CreatedAtField.IsString())
+					{
+						FDateTime::ParseIso8601(FUTF8ToTCHAR(CreatedAtField.AsString()).Get(), CreatedAt);
+					}
+					else if (CreatedAtField.IsDateTime())
+					{
+						CreatedAt = CreatedAtField.AsDateTime();
+					}
+				}
+
+				FString Configuration;
+				if (FCbFieldView ConfigurationField = BuildRecord.Metadata.FindIgnoreCase("Configuration"); ConfigurationField.HasValue() && !ConfigurationField.HasError())
+				{
+					Configuration = FUTF8ToTCHAR(ConfigurationField.AsString());
+				}
+
+				FString BuildVersionSuffix;
+				if (FCbFieldView BuildVersionSuffixField = BuildRecord.Metadata.FindIgnoreCase("BuildVersionSuffix"); BuildVersionSuffixField.HasValue() && !BuildVersionSuffixField.HasError())
+				{
+					BuildVersionSuffix = FUTF8ToTCHAR(BuildVersionSuffixField.AsString());
+					BuildVersionSuffix.RemoveFromStart(TEXT("-"));
+				}
+
+
+				FString ItemName;
+				FCbFieldView GroupNameView = BuildRecord.Metadata.FindIgnoreCase("buildgroup");
+				if (!GroupNameView.HasValue())
+				{
+					GroupNameView = BuildRecord.Metadata.FindIgnoreCase("name");
+				}
+				if (FCbFieldView NameView = GroupNameView; NameView.HasValue() && !NameView.HasError())
+				{
+					// TODO: This name manipulation needs to be removed when the metadata is more consistent.
+					ItemName = FUTF8ToTCHAR(NameView.AsString());
+				}
+				else
+				{
+					ItemName = *WriteToString<64>(BuildRecord.BuildId);
+				}
+
+				BuildGroup->Namespace = BuildState.Namespace;
+				BuildGroup->DisplayName = ItemName;
+				BuildGroup->CommitIdentifier = CommitIdentifier;
+				BuildGroup->Changelist = FCString::Atoi(*CommitIdentifier);
+				BuildGroup->Category = Category;
+				BuildGroup->CreatedAt = CreatedAt;
+				BuildGroup->Job = Job;
+				BuildGroup->Suffix = BuildVersionSuffix;
+				BuildGroup->bIsPreflight = bIsPreflight;
+			}
+
+			BuildGroup->NamedArtifacts.FindOrAdd(ArtifactName, MoveTemp(BuildRecord));
+		}
+	}
+
+	TArray<TSharedPtr<FBuildGroup>> BuildGroups;
+	KeyedGroups.GenerateValueArray(BuildGroups);
+
+	Result.Reserve(BuildGroups.Num());
+	for (const TSharedPtr<FBuildGroup>& BuildGroup : BuildGroups)
+	{
+		Result.Add(BuildGroup.ToSharedRef());
+	}
+
+	Result.Sort( [](const TSharedRef<FBuildListRetriever::FBuildGroup>& A, const TSharedRef<FBuildListRetriever::FBuildGroup>& B)
+	{
+		return A->CreatedAt > B->CreatedAt;
+	});
+
+	return MoveTemp(Result);
+}
+
+
+
+bool FBuildListRetriever::GetBucketInfo( const FString& BucketId, FBucketInfo& OutBucketInfo )
+{
+	TArray<FString> BucketItems;
+	if (BucketId.ParseIntoArray(BucketItems, TEXT(".")) != 4)
+	{
+		return false;
+	}
+
+	OutBucketInfo.Project = BucketItems[0];
+	OutBucketInfo.BuildType = BucketItems[1];
+	OutBucketInfo.Stream = BucketItems[2];
+	OutBucketInfo.Platform = BucketItems[3];
+	return true;
+}
+
+
+
+TMap<FString,TSharedRef<FBuildListRetriever::FBuildType, ESPMode::ThreadSafe>> FBuildListRetriever::GenerateBuildTypesMap() const
+{
+	// gather all build type names
+	TSet<FString> BuildTypeNames;
+	for (const auto& Pair : ServiceInstance->GetNamespacesAndBuckets())
+	{
+		const FString& Namespace = Pair.Key;
+		const FString& BucketId = Pair.Value;
+
+		FBucketInfo BucketInfo;
+		if (GetBucketInfo(BucketId, BucketInfo))
+		{
+			if (BucketInfo.Stream == Stream && BucketInfo.Project == ProjectName)
+			{
+				BuildTypeNames.Add(BucketInfo.BuildType);
+			}
+		}
+	}
+
+	// build the sorted list of build types
+	TMap<FString,TSharedRef<FBuildType>> Result;
+	for (const FString& BuildTypeName : BuildTypeNames)
+	{
+		Result.Add( BuildTypeName, MakeBuildType(BuildTypeName) );
+	}
+	Result.ValueSort( [](const TSharedRef<FBuildType>& A, const TSharedRef<FBuildType>& B)
+	{
+		if (A->Order != B->Order)
+		{
+			return A->Order >= B->Order;
+		}
+
+		return A->DisplayName.ToString() < B->DisplayName.ToString();
+	});
+
+
+	return Result;
+}
+
+
+
+TSharedRef<FBuildListRetriever::FBuildType, ESPMode::ThreadSafe> FBuildListRetriever::MakeBuildType(const FString& InBuildTypeName) const
+{
+	struct FKnownBuildType
+	{
+		FRegexPattern Pattern;
+		EBuildType Type;
+		FText DisplayName;
+		int Order;
+	};
+	static const FKnownBuildType KnownBuildTypes[] =
+	{
+		{ FRegexPattern(TEXT(".*oplog-?(.*)")),				EBuildType::Oplog,						LOCTEXT("BuildSelection_BuildType_Oplog", "Cook Snapshot"), 0 },
+		{ FRegexPattern(TEXT(".*packaged-?build-?(.*)")),	EBuildType::PackagedBuild,				LOCTEXT("BuildSelection_BuildType_PackagedBuild", "Packaged Build"), 1},
+		{ FRegexPattern(TEXT(".*staged-?(.*?)-?build.*")),	EBuildType::StagedBuild,				LOCTEXT("BuildSelection_BuildType_StagedBuild", "Staged Build"), 2},
+		{ FRegexPattern(TEXT(".*ugs-?pcb-?(.*)")),			EBuildType::EditorPreCompiledBinary,	LOCTEXT("BuildSelection_BuildType_PreCompiledBinary", "Editor Pre-compiled Binary"), 0},
+		{ FRegexPattern(TEXT(".*installed-?build-?(.*)")),	EBuildType::EditorPreCompiledBinary,	LOCTEXT("BuildSelection_BuildType_PreCompiledBinary", "Editor Pre-compiled Binary"), 0},
+	};
+
+
+	TSharedRef<FBuildType, ESPMode::ThreadSafe> Result = MakeShared<FBuildType, ESPMode::ThreadSafe>();
+
+	for (int32 Index = 0; Index < UE_ARRAY_COUNT(KnownBuildTypes); Index++)
+	{
+		const FKnownBuildType& KnownBuildType = KnownBuildTypes[Index];
+		FRegexMatcher Matcher(KnownBuildType.Pattern, InBuildTypeName);
+		if (Matcher.FindNext())
+		{
+			Result->Type = KnownBuildType.Type;
+			Result->Order = KnownBuildType.Order;
+			Result->DisplayName = KnownBuildType.DisplayName;
+
+			FString OptionalCaptureGroupString = Matcher.GetCaptureGroup(1);
+			if (!OptionalCaptureGroupString.IsEmpty())
+			{
+				Result->DisplayName = FText::Format(LOCTEXT("BuildSelection_KnownBuildTypeWithCaptureGroupFormat", "{0} ({1})"), KnownBuildType.DisplayName, FText::FromString(OptionalCaptureGroupString));
+			}
+
+			return Result;
+		}
+	}
+
+	// unknown build tye
+	Result->Type = EBuildType::Unknown;
+	Result->DisplayName = FText::FromString(*InBuildTypeName);
+	Result->Order = 0;
+	
+	return Result;
+}
+
+
+
 } // namespace UE::Zen::Build
+
+#undef LOCTEXT_NAMESPACE
diff --git a/Engine/Source/Developer/Zen/Public/Experimental/BuildListRetriever.h b/Engine/Source/Developer/Zen/Public/Experimental/BuildListRetriever.h
--- a/Engine/Source/Developer/Zen/Public/Experimental/BuildListRetriever.h
+++ b/Engine/Source/Developer/Zen/Public/Experimental/BuildListRetriever.h
@@ -5,6 +5,9 @@
 #include "Templates/SharedPointer.h"
 #include "Experimental/BuildServerInterface.h"
 #include "ZenBuildUtils.h"
+#include "Internationalization/Text.h"
+#include "Internationalization/Regex.h"
+#include "Misc/DateTime.h"
 
 #define UE_API ZEN_API
 
@@ -17,16 +20,73 @@ namespace UE::Zen::Build
 class FBuildListRetriever : public TSharedFromThis<FBuildListRetriever>
 {
 public:
+	struct FBuildGroup : public TSharedFromThis<FBuildGroup>
+	{
+		FString Namespace;
+		FString DisplayName;
+		FString CommitIdentifier;
+		int32 Changelist;
+		FDateTime CreatedAt;
+		FString Category;
+		FString Job;
+		FString Suffix;
+		TMap<FString, UE::Zen::Build::FBuildServiceInstance::FBuildRecord> NamedArtifacts;
+		bool bIsPreflight;
+	};
+
+
 	UE_API FBuildListRetriever();
 	typedef TMap<FString, TArray<int32>> FPlatformBuildsMap;
+	typedef TUniqueFunction<void()> FOnReadyForQuery;
 	typedef TUniqueFunction<void(const FPlatformBuildsMap&)> FOnQueriesComplete;
+	typedef TUniqueFunction<void(const TArrayView<TSharedRef<FBuildGroup>>&)> FOnGroupQueriesComplete;
 	UE_API void ConnectToBuildService();
+	UE_API bool IsConnected() const;
 	UE_API bool QueryBuilds(const FString& InProjectName, const FString& InBuildType, const FString& InStream, FOnQueriesComplete&& InOnQueriesCompleteCallback);
+	UE_API bool QueryBuildGroups(const FString& InProjectName, const FString& InBuildType, const FString& InStream, FOnGroupQueriesComplete&& InOnQueriesCompleteCallback);
+	UE_API bool IsReadyForQuery() const;
+	UE_API void SetOnReadyForQuery( FOnReadyForQuery&& InOnReadyForQuery );
+	UE_API void Refresh();
+
+	UE_API TSharedRef<FBuildServiceInstance> GetServiceInstance() const;
+
+
+	enum class EBuildType
+	{
+		Oplog,
+		StagedBuild,
+		PackagedBuild,
+		EditorPreCompiledBinary,
+		EditorInstalledBuild,
+		Unknown,
+		Count
+	};
+	struct FBuildType : TSharedFromThis<FBuildType>
+	{
+		EBuildType Type;
+		FText DisplayName;
+		int Order;
+	};
+
+	UE_API TMap<FString,TSharedRef<FBuildType, ESPMode::ThreadSafe>> GenerateBuildTypesMap() const;
+
+
+	struct FBucketInfo
+	{
+		FString Project;
+		FString BuildType;
+		FString Stream;
+		FString Platform;
+	};
+	static UE_API bool GetBucketInfo( const FString& BucketId, FBucketInfo& OutBucketInfo );
 
 private:
 	void OnNamespacesAndBucketsRefreshed();
-	void IssueBuildQueries();
+	void IssueBuildQueries(bool bIsGroupQuery);
 	void ExtractBuildsFromQueryResult(FListBuildsState& ListBuildsState);
+	TArray<TSharedRef<FBuildGroup>> ExtractBuildGroupsFromQueryResult(FListBuildsState& ListBuildsState);
+	TSharedRef<FBuildType, ESPMode::ThreadSafe> MakeBuildType(const FString& InBuildTypeName) const;
+
 
 	FString ProjectName;
 	FString BuildType;
@@ -34,7 +94,9 @@ private:
 
 	mutable TSharedPtr<FBuildServiceInstance> ServiceInstance;
 	FPlatformBuildsMap PerPlatformBuilds;
+	FOnReadyForQuery OnReadyForQuery;
 	FOnQueriesComplete OnQueriesCompleteCallback;
+	FOnGroupQueriesComplete OnGroupQueriesCompleteCallback;
 	bool ReadyToIssueBuildQuery = false; // Should only be accessed from the game thread
 	bool bConnected = false;
 };
