Shifted the Zen unit test from Basic tag to Slow tag since it involves a lot of subprocess launches and slows down FoundationTest execution.  This will still be run on the build farm.

#Synced-CL 49059893

diff --git a/Engine/Source/Developer/Zen/Tests/ZenTest.cpp b/Engine/Source/Developer/Zen/Tests/ZenTest.cpp
--- a/Engine/Source/Developer/Zen/Tests/ZenTest.cpp
+++ b/Engine/Source/Developer/Zen/Tests/ZenTest.cpp
@@ -14,7 +14,7 @@
 namespace UE
 {
 
-TEST_CASE("Zen::ZenServerInterface", "[Zen][Basic]")
+TEST_CASE("Zen::ZenServerInterface", "[Zen][Slow]")
 {
 	using namespace UE::Zen;
 	uint16 DefaultTestPort = 8559;
