Change the meaning of ConnectTimeoutMS on FZenHttpRequest's constructor (internal) to mean TimeoutMS instead.  This helps avoid hangs when we connect, but then the other side doesn't respond after connecting (or the response is in some way disrupted).  Existing uses of the connect timeout have been changed to use the overall timeout instead as they generally expect quick responses after connecting (low amounts of data are sent).

#Synced-CL 49059893

diff --git a/Engine/Source/Developer/Zen/Internal/ZenServerHttp.h b/Engine/Source/Developer/Zen/Internal/ZenServerHttp.h
--- a/Engine/Source/Developer/Zen/Internal/ZenServerHttp.h
+++ b/Engine/Source/Developer/Zen/Internal/ZenServerHttp.h
@@ -88,7 +88,7 @@ inline FStringView GetMimeType(EContentType Type)
 class FZenHttpRequest
 {
 public:
-	ZEN_API FZenHttpRequest(const FZenServiceEndpoint& InDomain, bool bInLogErrors, uint32 ConnectTimeoutMS = 0);
+	ZEN_API FZenHttpRequest(const FZenServiceEndpoint& InDomain, bool bInLogErrors, uint32 TimeoutMS = 0);
 	ZEN_API ~FZenHttpRequest();
 
 	/**
diff --git a/Engine/Source/Developer/Zen/Private/ZenServerHttp.cpp b/Engine/Source/Developer/Zen/Private/ZenServerHttp.cpp
--- a/Engine/Source/Developer/Zen/Private/ZenServerHttp.cpp
+++ b/Engine/Source/Developer/Zen/Private/ZenServerHttp.cpp
@@ -68,7 +68,7 @@ namespace UE::Zen {
 #endif //PLATFORM_WINDOWS
 	};
 
-	FZenHttpRequest::FZenHttpRequest(const FZenServiceEndpoint& InEndpoint, bool bInLogErrors, uint32 ConnectTimeoutMS)
+	FZenHttpRequest::FZenHttpRequest(const FZenServiceEndpoint& InEndpoint, bool bInLogErrors, uint32 TimeoutMS)
 		: bLogErrors(bInLogErrors)
 		, Domain(InEndpoint.GetURL())
 		, Endpoint(InEndpoint)
@@ -80,9 +80,9 @@ namespace UE::Zen {
 		}
 		Curl = curl_easy_init();
 		Reset();
-		if (ConnectTimeoutMS != 0)
+		if (TimeoutMS != 0)
 		{
-			curl_easy_setopt(Curl, CURLOPT_CONNECTTIMEOUT_MS, ConnectTimeoutMS);
+			curl_easy_setopt(Curl, CURLOPT_TIMEOUT_MS, TimeoutMS);
 		}
 	}
 
diff --git a/Engine/Source/Developer/Zen/Private/ZenServerInterface.cpp b/Engine/Source/Developer/Zen/Private/ZenServerInterface.cpp
--- a/Engine/Source/Developer/Zen/Private/ZenServerInterface.cpp
+++ b/Engine/Source/Developer/Zen/Private/ZenServerInterface.cpp
@@ -1554,6 +1554,7 @@ FServiceSettings::TryApplyAutoLaunchOverride()
 
 uint16 FZenServiceInstance::AutoLaunchedPort = 0;
 uint32 FZenServiceInstance::AutoLaunchedPid = 0;
+const uint32 FZenServiceInstance::DefaultSmallOperationTimeoutMS = 500;
 
 static bool
 IsZenProcessUsingEffectivePort(uint16 EffectiveListenPort)
@@ -2331,7 +2332,7 @@ FZenServiceInstance::PingService(uint32_t AttemptCount, bool Silent)
 	{
 		Attempt++;
 
-		Zen::FZenHttpRequest Request(Endpoint, false, 500);
+		Zen::FZenHttpRequest Request(Endpoint, false, DefaultSmallOperationTimeoutMS);
 		Zen::FZenHttpRequest::Result Result = Request.PerformBlockingDownload(TEXTVIEW("health/ready"), nullptr, Zen::EContentType::Text);
 		
 		if (Result == Zen::FZenHttpRequest::Result::Success && Zen::IsSuccessCode(Request.GetResponseCode()))
@@ -4061,7 +4062,7 @@ FZenServiceInstance::GetCacheStats(FZenCacheStats& Stats)
 		// We've not got any requests in flight and we've met a given time requirement for requests
 	CacheStatsRequest = Async(ThreadPool, [&Endpoint = Endpoint]
 		{
-			UE::Zen::FZenHttpRequest Request(Endpoint, false, 500);
+			UE::Zen::FZenHttpRequest Request(Endpoint, false, DefaultSmallOperationTimeoutMS);
 
 			TArray64<uint8> GetBuffer;
 			FZenHttpRequest::Result Result = Request.PerformBlockingDownload(TEXTVIEW("/stats/z$"), &GetBuffer, Zen::EContentType::CbObject);
@@ -4115,7 +4116,7 @@ FZenServiceInstance::GetProjectStats(FZenProjectStats& Stats)
 			// We've not got any requests in flight and we've met a given time requirement for requests
 	ProjectStatsRequest = Async(ThreadPool, [&Endpoint = Endpoint]
 		{
-			UE::Zen::FZenHttpRequest Request(Endpoint, false, 500);
+			UE::Zen::FZenHttpRequest Request(Endpoint, false, DefaultSmallOperationTimeoutMS);
 
 			TArray64<uint8> GetBuffer;
 			FZenHttpRequest::Result Result = Request.PerformBlockingDownload(TEXTVIEW("/stats/prj"), &GetBuffer, Zen::EContentType::CbObject);
@@ -4168,7 +4169,7 @@ FZenServiceInstance::GetGCStatus(FGCStatus& Status)
 		// We've not got any requests in flight and we've met a given time requirement for requests
 		GCStatusRequest = Async(ThreadPool, [this]
 			{
-				UE::Zen::FZenHttpRequest Request(Endpoint, false, 500);
+				UE::Zen::FZenHttpRequest Request(Endpoint, false, DefaultSmallOperationTimeoutMS);
 
 				TArray64<uint8> GetBuffer;
 				FZenHttpRequest::Result Result = Request.PerformBlockingDownload(TEXTVIEW("/admin/gc"), &GetBuffer, Zen::EContentType::CbObject);
@@ -4193,7 +4194,7 @@ FZenServiceInstance::GetGCStatus(FGCStatus& Status)
 bool 
 FZenServiceInstance::RequestGC(const bool* OverrideCollectSmallObjects, const uint32* OverrideMaxCacheDuration)
 {
-	UE::Zen::FZenHttpRequest Request(Endpoint, /*bLogErrors*/ false);
+	UE::Zen::FZenHttpRequest Request(Endpoint, /*bLogErrors*/ false, DefaultSmallOperationTimeoutMS);
 
 	TCHAR Separators[] = {TEXT('?'), TEXT('&')};
 	int32 SeparatorIndex = 0;
@@ -4488,7 +4489,7 @@ FZenServiceInstance::GetWorkspaces(FZenWorkspaces& Workspaces) const
 
 	WorkspacesRequest = Async(ThreadPool, [&Endpoint = Endpoint]
 	{
-		UE::Zen::FZenHttpRequest Request(Endpoint, false, 500);
+		UE::Zen::FZenHttpRequest Request(Endpoint, false, DefaultSmallOperationTimeoutMS);
 
 		FZenHttpRequest::Result Result = Request.PerformBlockingDownload(TEXTVIEW("/ws"), nullptr, Zen::EContentType::CbObject);
 
diff --git a/Engine/Source/Developer/Zen/Public/Experimental/ZenServerInterface.h b/Engine/Source/Developer/Zen/Public/Experimental/ZenServerInterface.h
--- a/Engine/Source/Developer/Zen/Public/Experimental/ZenServerInterface.h
+++ b/Engine/Source/Developer/Zen/Public/Experimental/ZenServerInterface.h
@@ -368,6 +368,7 @@ private:
 	FZenServiceEndpoint Endpoint;
 	static uint16 AutoLaunchedPort;
 	static uint32 AutoLaunchedPid;
+	static const uint32 DefaultSmallOperationTimeoutMS;
 	bool bHasLaunchedLocal = false;
 	bool bHasLaunchedSystemService = false;
 	bool bIsRunningLocally = true;
