Downgrade warnings when attempting to ping autolaunched zenserver

[RN] none
#Synced-CL 48985904

diff --git a/Engine/Source/Developer/Zen/Private/ZenServerInterface.cpp b/Engine/Source/Developer/Zen/Private/ZenServerInterface.cpp
--- a/Engine/Source/Developer/Zen/Private/ZenServerInterface.cpp
+++ b/Engine/Source/Developer/Zen/Private/ZenServerInterface.cpp
@@ -2350,7 +2350,7 @@ FZenServiceInstance::PingService(uint32_t AttemptCount, bool Silent)
 			{
 				if (!Silent)
 				{
-					UE_LOGFMT(LogZenServiceInstance, Warning, "Unable to reach Local Unreal Zen Storage Server HTTP service at {Domain}: {Result}", Endpoint.GetURL(), Request.GetDetailedResultString(Endpoint.GetURL()));
+					UE_LOGFMT(LogZenServiceInstance, Display, "Unable to reach Local Unreal Zen Storage Server HTTP service at {Domain}: {Result}", Endpoint.GetURL(), Request.GetDetailedResultString(Endpoint.GetURL()));
 				}
 				break;
 			}
@@ -2620,7 +2620,7 @@ FZenServiceInstance::Initialize()
 							}
 							else
 							{
-								UE_LOG(LogZenServiceInstance, Warning, TEXT("Local ZenServer AutoLaunch completed but unable to reach service at %s"), *Endpoint.GetURL());
+								UE_LOG(LogZenServiceInstance, Display, TEXT("Local ZenServer AutoLaunch completed but unable to reach service at %s"), *Endpoint.GetURL());
 							}
 						}
 					}
