Add an overall Timeout alongside the ConnectTimeout offered by FZenHttpRequest.  Default usage adheres to past behavior (5 second connect timeout, infinite overall timeout), but the usages in ZenServerInstance use 0.5s connect timeout (as before) and either 10 minute or 1 minute overall timeout depending on the operation.

#Synced-CL 49059893

diff --git a/Engine/Source/Developer/Zen/Internal/ZenServerHttp.h b/Engine/Source/Developer/Zen/Internal/ZenServerHttp.h
--- a/Engine/Source/Developer/Zen/Internal/ZenServerHttp.h
+++ b/Engine/Source/Developer/Zen/Internal/ZenServerHttp.h
@@ -88,7 +88,7 @@ inline FStringView GetMimeType(EContentType Type)
 class FZenHttpRequest
 {
 public:
-	ZEN_API FZenHttpRequest(const FZenServiceEndpoint& InDomain, bool bInLogErrors, uint32 ConnectTimeoutMS = 0);
+	ZEN_API FZenHttpRequest(const FZenServiceEndpoint& InDomain, bool bInLogErrors, uint32 InConnectTimeoutMS = 0, uint32 InTimeoutMS = 0);
 	ZEN_API ~FZenHttpRequest();
 
 	/**
@@ -220,6 +220,8 @@ private:
 	size_t					BytesReceived = 0;
 	bool					bLogErrors = false;
 	bool					bResponseFormatValid = false;
+	uint32				ConnectTimeoutMS = 0;
+	uint32				TimeoutMS = 0;
 
 	const FCompositeBuffer*	ReadDataView = nullptr;
 	TArray64<uint8>*		WriteDataBufferPtr = nullptr;
diff --git a/Engine/Source/Developer/Zen/Private/ZenServerHttp.cpp b/Engine/Source/Developer/Zen/Private/ZenServerHttp.cpp
--- a/Engine/Source/Developer/Zen/Private/ZenServerHttp.cpp
+++ b/Engine/Source/Developer/Zen/Private/ZenServerHttp.cpp
@@ -68,8 +68,10 @@ namespace UE::Zen {
 #endif //PLATFORM_WINDOWS
 	};
 
-	FZenHttpRequest::FZenHttpRequest(const FZenServiceEndpoint& InEndpoint, bool bInLogErrors, uint32 ConnectTimeoutMS)
+	FZenHttpRequest::FZenHttpRequest(const FZenServiceEndpoint& InEndpoint, bool bInLogErrors, uint32 InConnectTimeoutMS, uint32 InTimeoutMS)
 		: bLogErrors(bInLogErrors)
+		, ConnectTimeoutMS(InConnectTimeoutMS)
+		, TimeoutMS(InTimeoutMS)
 		, Domain(InEndpoint.GetURL())
 		, Endpoint(InEndpoint)
 	{
@@ -80,10 +82,6 @@ namespace UE::Zen {
 		}
 		Curl = curl_easy_init();
 		Reset();
-		if (ConnectTimeoutMS != 0)
-		{
-			curl_easy_setopt(Curl, CURLOPT_CONNECTTIMEOUT_MS, ConnectTimeoutMS);
-		}
 	}
 
 	FZenHttpRequest::~FZenHttpRequest()
@@ -108,7 +106,11 @@ namespace UE::Zen {
 		curl_easy_reset(Curl);
 
 		// Options that are always set for all connections.
-		curl_easy_setopt(Curl, CURLOPT_CONNECTTIMEOUT, 5L);
+		curl_easy_setopt(Curl, CURLOPT_CONNECTTIMEOUT_MS, ConnectTimeoutMS > 0 ? ConnectTimeoutMS : 5000L);
+		if (TimeoutMS > 0)
+		{
+			curl_easy_setopt(Curl, CURLOPT_TIMEOUT_MS, TimeoutMS);
+		}
 		curl_easy_setopt(Curl, CURLOPT_EXPECT_100_TIMEOUT_MS, 0);
 		curl_easy_setopt(Curl, CURLOPT_NOSIGNAL, 1L);
 		curl_easy_setopt(Curl, CURLOPT_DNS_CACHE_TIMEOUT, -1L); // Don't re-resolve names mid-session
diff --git a/Engine/Source/Developer/Zen/Private/ZenServerInterface.cpp b/Engine/Source/Developer/Zen/Private/ZenServerInterface.cpp
--- a/Engine/Source/Developer/Zen/Private/ZenServerInterface.cpp
+++ b/Engine/Source/Developer/Zen/Private/ZenServerInterface.cpp
@@ -1554,6 +1554,8 @@ FServiceSettings::TryApplyAutoLaunchOverride()
 
 uint16 FZenServiceInstance::AutoLaunchedPort = 0;
 uint32 FZenServiceInstance::AutoLaunchedPid = 0;
+const uint32 FZenServiceInstance::DefaultSmallOperationConnectTimeoutMS = 500; // 0.5 second connection timeout
+const uint32 FZenServiceInstance::DefaultSmallOperationTimeoutMS = 60000; // 1 minute total duration timout
 
 static bool
 IsZenProcessUsingEffectivePort(uint16 EffectiveListenPort)
@@ -2326,12 +2328,15 @@ FZenServiceInstance::IsServiceReady()
 bool 
 FZenServiceInstance::PingService(uint32_t AttemptCount, bool Silent)
 {
+	// 10 minute duration timeout per ping.  If the server is wiping state,
+	// we may have to wait a long time for it after startup.
+	const uint32 DefaultPingOperationTimeoutMS = 600000;
 	uint32 Attempt = 0;
 	while (IsServiceRunning())
 	{
 		Attempt++;
 
-		Zen::FZenHttpRequest Request(Endpoint, false, 500);
+		Zen::FZenHttpRequest Request(Endpoint, false, DefaultSmallOperationConnectTimeoutMS, DefaultPingOperationTimeoutMS);
 		Zen::FZenHttpRequest::Result Result = Request.PerformBlockingDownload(TEXTVIEW("health/ready"), nullptr, Zen::EContentType::Text);
 		
 		if (Result == Zen::FZenHttpRequest::Result::Success && Zen::IsSuccessCode(Request.GetResponseCode()))
@@ -4061,7 +4066,7 @@ FZenServiceInstance::GetCacheStats(FZenCacheStats& Stats)
 		// We've not got any requests in flight and we've met a given time requirement for requests
 	CacheStatsRequest = Async(ThreadPool, [&Endpoint = Endpoint]
 		{
-			UE::Zen::FZenHttpRequest Request(Endpoint, false, 500);
+			UE::Zen::FZenHttpRequest Request(Endpoint, false, DefaultSmallOperationConnectTimeoutMS, DefaultSmallOperationTimeoutMS);
 
 			TArray64<uint8> GetBuffer;
 			FZenHttpRequest::Result Result = Request.PerformBlockingDownload(TEXTVIEW("/stats/z$"), &GetBuffer, Zen::EContentType::CbObject);
@@ -4115,7 +4120,7 @@ FZenServiceInstance::GetProjectStats(FZenProjectStats& Stats)
 			// We've not got any requests in flight and we've met a given time requirement for requests
 	ProjectStatsRequest = Async(ThreadPool, [&Endpoint = Endpoint]
 		{
-			UE::Zen::FZenHttpRequest Request(Endpoint, false, 500);
+			UE::Zen::FZenHttpRequest Request(Endpoint, false, DefaultSmallOperationConnectTimeoutMS, DefaultSmallOperationTimeoutMS);
 
 			TArray64<uint8> GetBuffer;
 			FZenHttpRequest::Result Result = Request.PerformBlockingDownload(TEXTVIEW("/stats/prj"), &GetBuffer, Zen::EContentType::CbObject);
@@ -4168,7 +4173,7 @@ FZenServiceInstance::GetGCStatus(FGCStatus& Status)
 		// We've not got any requests in flight and we've met a given time requirement for requests
 		GCStatusRequest = Async(ThreadPool, [this]
 			{
-				UE::Zen::FZenHttpRequest Request(Endpoint, false, 500);
+				UE::Zen::FZenHttpRequest Request(Endpoint, false, DefaultSmallOperationConnectTimeoutMS, DefaultSmallOperationTimeoutMS);
 
 				TArray64<uint8> GetBuffer;
 				FZenHttpRequest::Result Result = Request.PerformBlockingDownload(TEXTVIEW("/admin/gc"), &GetBuffer, Zen::EContentType::CbObject);
@@ -4193,7 +4198,7 @@ FZenServiceInstance::GetGCStatus(FGCStatus& Status)
 bool 
 FZenServiceInstance::RequestGC(const bool* OverrideCollectSmallObjects, const uint32* OverrideMaxCacheDuration)
 {
-	UE::Zen::FZenHttpRequest Request(Endpoint, /*bLogErrors*/ false);
+	UE::Zen::FZenHttpRequest Request(Endpoint, /*bLogErrors*/ false, DefaultSmallOperationConnectTimeoutMS, DefaultSmallOperationTimeoutMS);
 
 	TCHAR Separators[] = {TEXT('?'), TEXT('&')};
 	int32 SeparatorIndex = 0;
@@ -4488,7 +4493,7 @@ FZenServiceInstance::GetWorkspaces(FZenWorkspaces& Workspaces) const
 
 	WorkspacesRequest = Async(ThreadPool, [&Endpoint = Endpoint]
 	{
-		UE::Zen::FZenHttpRequest Request(Endpoint, false, 500);
+		UE::Zen::FZenHttpRequest Request(Endpoint, false, DefaultSmallOperationConnectTimeoutMS, DefaultSmallOperationTimeoutMS);
 
 		FZenHttpRequest::Result Result = Request.PerformBlockingDownload(TEXTVIEW("/ws"), nullptr, Zen::EContentType::CbObject);
 
diff --git a/Engine/Source/Developer/Zen/Public/Experimental/ZenServerInterface.h b/Engine/Source/Developer/Zen/Public/Experimental/ZenServerInterface.h
--- a/Engine/Source/Developer/Zen/Public/Experimental/ZenServerInterface.h
+++ b/Engine/Source/Developer/Zen/Public/Experimental/ZenServerInterface.h
@@ -368,6 +368,8 @@ private:
 	FZenServiceEndpoint Endpoint;
 	static uint16 AutoLaunchedPort;
 	static uint32 AutoLaunchedPid;
+	static const uint32 DefaultSmallOperationConnectTimeoutMS;
+	static const uint32 DefaultSmallOperationTimeoutMS;
 	bool bHasLaunchedLocal = false;
 	bool bHasLaunchedSystemService = false;
 	bool bIsRunningLocally = true;
