Reduce the time spent running "zen.exe service status" from unit tests that repeatedly start/shutdown zenserver.

#Synced-CL 49552467

diff --git a/Engine/Source/Developer/Zen/Private/ZenServerInterface.cpp b/Engine/Source/Developer/Zen/Private/ZenServerInterface.cpp
--- a/Engine/Source/Developer/Zen/Private/ZenServerInterface.cpp
+++ b/Engine/Source/Developer/Zen/Private/ZenServerInterface.cpp
@@ -2566,18 +2566,23 @@ FZenServiceInstance::Initialize()
 	IFileManager& FileManager = IFileManager::Get();
 	FileManager.MakeDirectory(FPlatformProcess::ApplicationSettingsDir());
 
-	if (!IsDefaultInstanceSystemService() && IsDefaultInstanceAutoLaunch())
+	static bool bFirstTime = true;
+	if (bFirstTime)
 	{
-		FSystemWideCriticalSection InstallCriticalSection(TEXT("ZenServerInstall"), FTimespan::FromSeconds(FApp::IsUnattended() ? 30.0 : 10.0));
-		if (!InstallCriticalSection.IsValid())
+		if (!IsDefaultInstanceSystemService() && IsDefaultInstanceAutoLaunch())
 		{
-			PromptUserOfFailedSystemServiceUninstall();
-		}
+			FSystemWideCriticalSection InstallCriticalSection(TEXT("ZenServerInstall"), FTimespan::FromSeconds(FApp::IsUnattended() ? 30.0 : 10.0));
+			if (!InstallCriticalSection.IsValid())
+			{
+				PromptUserOfFailedSystemServiceUninstall();
+			}
 
-		if (!UninstallSystemService(GetInTreeUtilityPath()))
-		{
-			PromptUserOfFailedSystemServiceUninstall();
+			if (!UninstallSystemService(GetInTreeUtilityPath()))
+			{
+				PromptUserOfFailedSystemServiceUninstall();
+			}
 		}
+		bFirstTime = false;
 	}
 
 	if (Settings.IsAutoLaunch())
