Provide clearer warnings when we fail to gracefully shut down zenserver.

#Synced-CL 49552467

diff --git a/Engine/Source/Developer/Zen/Private/ZenServerInterface.cpp b/Engine/Source/Developer/Zen/Private/ZenServerInterface.cpp
--- a/Engine/Source/Developer/Zen/Private/ZenServerInterface.cpp
+++ b/Engine/Source/Developer/Zen/Private/ZenServerInterface.cpp
@@ -1616,15 +1616,15 @@ static bool ShutdownZenServerProcess(int Pid, double MaximumWaitDurationSeconds
 		return true;
 	}
 
-	UE_LOG(LogZenServiceInstance, Display, TEXT("Attempting termination of process with pid %d"), Pid);
+	UE_LOG(LogZenServiceInstance, Display, TEXT("After failed shut down request, attempting forced termination of process with pid %d"), Pid);
 	if (ZenServerState::Terminate(Pid))
 	{
-		UE_LOG(LogZenServiceInstance, Warning, TEXT("Terminated process with pid %d"), Pid);
+		UE_LOG(LogZenServiceInstance, Warning, TEXT("After failed shut down request, forced to terminate process with pid %d"), Pid);
 		return true;
 	}
 	if (ZenServerState::IsProcessRunning(Pid))
 	{
-		UE_LOG(LogZenServiceInstance, Warning, TEXT("Failed to terminate process with pid %d"), Pid);
+		UE_LOG(LogZenServiceInstance, Warning, TEXT("Failed to forcefully terminate process with pid %d"), Pid);
 		return false;
 	}
 	UE_LOG(LogZenServiceInstance, Display, TEXT("No longer any active process with pid %d"), Pid);
