Clarify messaging around shutdown of zenserver.  Also downgrade the warning that can occur if there is a zenserver process running but not in the state table and we terminate it.  This can happen if we encounter the process as it is shutting down and has removed itself from the state table, but before the process is completely gone.  We should not warn about this, so only print a display verbosity message about it.

#Synced-CL 49552467

diff --git a/Engine/Source/Developer/Zen/Private/ZenServerInterface.cpp b/Engine/Source/Developer/Zen/Private/ZenServerInterface.cpp
--- a/Engine/Source/Developer/Zen/Private/ZenServerInterface.cpp
+++ b/Engine/Source/Developer/Zen/Private/ZenServerInterface.cpp
@@ -1583,13 +1583,13 @@ static bool ShutdownZenServerProcess(int Pid, double MaximumWaitDurationSeconds
 {
 	const ZenServerState ServerState(/* ReadOnly */true);
 	const ZenServerState::ZenServerEntry* Entry = ServerState.LookupByPid(Pid);
+	uint64 ZenShutdownWaitStartTime = FPlatformTime::Cycles64();
 	if (Entry)
 	{
 		uint16 EffectivePort = Entry->EffectiveListenPort.load(std::memory_order_relaxed);
 		UE_LOG(LogZenServiceInstance, Display, TEXT("Requesting shut down of zenserver process %d runnning on effective port %u"), Pid, EffectivePort);
 		if (RequestZenShutdownOnEffectivePort(EffectivePort))
 		{
-			uint64 ZenShutdownWaitStartTime = FPlatformTime::Cycles64();
 			while (ZenServerState::IsProcessRunning(Pid))
 			{
 				double ZenShutdownWaitDuration = FPlatformTime::ToSeconds64(FPlatformTime::Cycles64() - ZenShutdownWaitStartTime);
@@ -1616,12 +1616,32 @@ static bool ShutdownZenServerProcess(int Pid, double MaximumWaitDurationSeconds
 		return true;
 	}
 
-	UE_LOG(LogZenServiceInstance, Display, TEXT("After failed shut down request, attempting forced termination of process with pid %d"), Pid);
-	if (ZenServerState::Terminate(Pid))
+	// There may be no entry left in the state table, but the process might be wrapping up its shutdown sequence.
+	// Allow use of the remaining wait duration for the process to finish exiting gracefully.
+	while (ZenServerState::IsProcessRunning(Pid))
 	{
-		UE_LOG(LogZenServiceInstance, Warning, TEXT("After failed shut down request, forced to terminate process with pid %d"), Pid);
-		return true;
+		double ZenShutdownWaitDuration = FPlatformTime::ToSeconds64(FPlatformTime::Cycles64() - ZenShutdownWaitStartTime);
+		if (ZenShutdownWaitDuration < MaximumWaitDurationSeconds)
+		{
+			FPlatformProcess::Sleep(0.01f);
+		}
+		else
+		{
+			UE_LOG(LogZenServiceInstance, Warning, TEXT("Timed out waiting for shut down of running service with pid %d"), Pid);
+			break;
+		}
 	}
+
+	if (ZenServerState::IsProcessRunning(Pid))
+	{
+		UE_LOG(LogZenServiceInstance, Display, TEXT("After failed shut down request, attempting forced termination of process with pid %d"), Pid);
+		if (ZenServerState::Terminate(Pid))
+		{
+			UE_LOG(LogZenServiceInstance, Warning, TEXT("After failed shut down request, forced to terminate process with pid %d"), Pid);
+			return true;
+		}
+	}
+
 	if (ZenServerState::IsProcessRunning(Pid))
 	{
 		UE_LOG(LogZenServiceInstance, Warning, TEXT("Failed to forcefully terminate process with pid %d"), Pid);
@@ -1637,7 +1657,7 @@ static bool ShutDownZenServerProcessExecutable(const FString& ExecutablePath, do
 	uint32_t Pid = 0;
 	while (ZenServerState::FindRunningProcessId(*ExecutablePath, &Pid))
 	{
-		UE_LOG(LogZenServiceInstance, Display, TEXT("Attempting to shut down of zenserver executable '%s' process with pid %d"), *ExecutablePath, Pid);
+		UE_LOG(LogZenServiceInstance, Display, TEXT("Attempting to shut down zenserver executable '%s' process with pid %d"), *ExecutablePath, Pid);
 		double ZenShutdownWaitDuration = FPlatformTime::ToSeconds64(FPlatformTime::Cycles64() - ZenShutdownWaitStartTime);
 		if (ShutdownZenServerProcess(Pid, MaximumWaitDurationSeconds - ZenShutdownWaitDuration))
 		{
