BuildStorageTool: Add transfer options (Clean, Verify, Force, Scavenge, BoostWorkers) and the ability to specify a host override to a regional instance.  Also set in place the ability (not exposed yet) to download parts by id/name and to provide a download specification for selection of individual files to download.

#Synced-CL 51039791

diff --git a/Engine/Source/Developer/Zen/Private/BuildServerInterface.cpp b/Engine/Source/Developer/Zen/Private/BuildServerInterface.cpp
--- a/Engine/Source/Developer/Zen/Private/BuildServerInterface.cpp
+++ b/Engine/Source/Developer/Zen/Private/BuildServerInterface.cpp
@@ -15,6 +15,7 @@
 #include "Http/HttpHostBuilder.h"
 #include "Misc/App.h"
 #include "Misc/ConfigCacheIni.h"
+#include "Misc/FileHelper.h"
 #include "Misc/MonitoredProcess.h"
 #include "Misc/OutputDeviceFile.h"
 #include "Misc/Paths.h"
@@ -71,6 +72,11 @@ struct FBuildTransferState : public TSharedFromThis<FBuildTransferState, ESPMode
 	FString IncludeFilter;
 	FString ExcludeFilter;
 	FName TargetPlatformName = NAME_None;
+	FString DownloadSpecJSONContents;
+	TArray<FString> PartIds;
+	TArray<FString> PartNames;
+	FString HostOverride;
+	EBuildTransferRequestFlags RequestFlags;
 	double StartTime=0;
 
 	std::atomic<bool> bCancelRequested = false;
@@ -542,7 +548,12 @@ FBuildServiceInstance::StartBuildTransfer(
 	FStringView Namespace,
 	FStringView Bucket,
 	FStringView IncludeFilter,
-	FStringView ExcludeFilter)
+	FStringView ExcludeFilter,
+	FString&& DownloadSpecJSONContents,
+	TArray<FString>&& PartIds,
+	TArray<FString>&& PartNames,
+	FStringView HostOverride,
+	EBuildTransferRequestFlags RequestFlags)
 {
 	TSharedPtr<Private::FBuildTransferState> NewState = MakeShared<Private::FBuildTransferState>();
 	NewState->Type = FBuildTransfer::EType::Files;
@@ -553,6 +564,11 @@ FBuildServiceInstance::StartBuildTransfer(
 	NewState->Destination = DestinationFolder;
 	NewState->IncludeFilter = IncludeFilter;
 	NewState->ExcludeFilter = ExcludeFilter;
+	NewState->DownloadSpecJSONContents = DownloadSpecJSONContents;
+	NewState->PartIds = MoveTemp(PartIds);
+	NewState->PartNames = MoveTemp(PartNames);
+	NewState->HostOverride = HostOverride;
+	NewState->RequestFlags = RequestFlags;
 	NewState->StartTime = FPlatformTime::Seconds();
 	BuildTransferThreadStates.Enqueue(NewState.ToSharedRef());
 	KickBuildTransferThread();
@@ -573,7 +589,9 @@ FBuildServiceInstance::StartOplogBuildTransfer(
 	FStringView ProjectFilePath,
 	FStringView Namespace,
 	FStringView Bucket,
-	FName TargetPlatformName)
+	FName TargetPlatformName,
+	FStringView HostOverride,
+	EBuildTransferRequestFlags RequestFlags)
 {
 	TSharedPtr<Private::FBuildTransferState> NewState = MakeShared<Private::FBuildTransferState>();
 	NewState->Type = FBuildTransfer::EType::Oplog;
@@ -587,6 +605,8 @@ FBuildServiceInstance::StartOplogBuildTransfer(
 	NewState->EngineDir = EngineDir;
 	NewState->ProjectFilePath = ProjectFilePath;
 	NewState->TargetPlatformName = TargetPlatformName;
+	NewState->HostOverride = HostOverride;
+	NewState->RequestFlags = RequestFlags;
 	BuildTransferThreadStates.Enqueue(NewState.ToSharedRef());
 	KickBuildTransferThread();
 
@@ -610,10 +630,15 @@ FBuildServiceInstance::RepeatBuildTransfer(FBuildTransfer BuildTransfer)
 		NewState->Destination = BuildTransfer.State->Destination;
 		NewState->IncludeFilter = BuildTransfer.State->IncludeFilter;
 		NewState->ExcludeFilter = BuildTransfer.State->ExcludeFilter;
+		NewState->DownloadSpecJSONContents = BuildTransfer.State->DownloadSpecJSONContents;
+		NewState->PartIds = BuildTransfer.State->PartIds;
+		NewState->PartNames = BuildTransfer.State->PartNames;
 		NewState->RootDir = BuildTransfer.State->RootDir;
 		NewState->EngineDir = BuildTransfer.State->EngineDir;
 		NewState->ProjectFilePath = BuildTransfer.State->ProjectFilePath;
 		NewState->TargetPlatformName = BuildTransfer.State->TargetPlatformName;
+		NewState->HostOverride = BuildTransfer.State->HostOverride;
+		NewState->RequestFlags = BuildTransfer.State->RequestFlags;
 		NewState->StartTime = FPlatformTime::Seconds();
 	}
 	
@@ -1149,18 +1174,65 @@ FBuildServiceInstance::BuildTransferThreadLoop()
  		FString ZenUtilityPath = FPaths::ConvertRelativePathToFull(UE::Zen::GetLocalInstallUtilityPath());
  		FString OidcTokenExecutableFilename = FPaths::ConvertRelativePathToFull(FDesktopPlatformModule::TryGet()->GetOidcTokenExecutableFilename(FPaths::RootDir()));
 
+ 		FString DownloadSpecTempFile;
+ 		ON_SCOPE_EXIT
+ 		{
+ 			if (!DownloadSpecTempFile.IsEmpty())
+ 			{
+				IFileManager::Get().Delete(*DownloadSpecTempFile);
+ 			}
+ 		};
  		FString DestinationDirectory;
-		FString CommandLineArgs;
+		TStringBuilder<256> CommandLineBuilder;
 		switch (State->Type)
 		{
 		case FBuildTransfer::EType::Files:
 			DestinationDirectory = FPaths::ConvertRelativePathToFull(State->Destination);
-			CommandLineArgs = FString::Printf(TEXT("builds download --log-progress --host %s --local-path \"%s\" --namespace %s --bucket %s --build-id %s%s%s --oidctoken-exe-path \"%s\""),
+			CommandLineBuilder.Append(FString::Printf(TEXT("builds download --log-progress --host %s --local-path \"%s\" --namespace %s --bucket %s --build-id %s%s%s --oidctoken-exe-path \"%s\""),
 				*WriteToString<64>(EffectiveDomain), *DestinationDirectory, *State->Namespace, *State->Bucket,
 				*WriteToString<64>(State->BuildId),
 				State->IncludeFilter.IsEmpty() ? TEXT("") : *WriteToString<64>(TEXT(" --wildcard \""), State->IncludeFilter, TEXT("\"")),
 				State->ExcludeFilter.IsEmpty() ? TEXT("") : *WriteToString<64>(TEXT(" --exclude-wildcard \""), State->ExcludeFilter, TEXT("\"")),
-				*OidcTokenExecutableFilename);
+				*OidcTokenExecutableFilename));
+
+				if (!State->HostOverride.IsEmpty())
+				{
+					CommandLineBuilder.Appendf(TEXT(" --override-host \"%s\""), *State->HostOverride);
+				}
+
+				if (!State->PartIds.IsEmpty())
+				{
+					CommandLineBuilder.Appendf(TEXT(" --build-part-id \"%s\""), *FString::Join(State->PartIds, TEXT(",")));
+				}
+
+				if (!State->PartNames.IsEmpty())
+				{
+					CommandLineBuilder.Appendf(TEXT(" --build-part-name \"%s\""), *FString::Join(State->PartNames, TEXT(",")));
+				}
+
+				if (!State->DownloadSpecJSONContents.IsEmpty())
+				{
+					DownloadSpecTempFile = FPaths::CreateTempFilename(FPlatformProcess::UserTempDir(), TEXT("download"), TEXT(".json"));
+					if (FFileHelper::SaveStringToFile(State->DownloadSpecJSONContents, *DownloadSpecTempFile))
+					{
+						CommandLineBuilder.Appendf(TEXT(" --download-spec-path \"%s\""), *DownloadSpecTempFile);
+					}
+					else
+					{
+						UE_LOG(LogBuildServiceInstance, Warning, TEXT("Failed to write download specification file, full build will be downloaded."));
+						DownloadSpecTempFile.Empty();
+					}
+				}
+
+				if (EnumHasAnyFlags(State->RequestFlags, EBuildTransferRequestFlags::Scavenge))
+				{
+					CommandLineBuilder.Append(TEXT(" --enable-scavenge"));
+				}
+
+				if (EnumHasAnyFlags(State->RequestFlags, EBuildTransferRequestFlags::Verify))
+				{
+					CommandLineBuilder.Append(TEXT(" --verify"));
+				}
 			break;
 		case FBuildTransfer::EType::Oplog:
 			{
@@ -1257,17 +1329,44 @@ FBuildServiceInstance::BuildTransferThreadLoop()
 					InvokeZenUtilitySync(FString::Printf(TEXT("oplog-create --project \"%s\" --oplog \"%s\""),
 						*DestinationProjectId, *DestinationOplogId));
 				}
-
-				CommandLineArgs = FString::Printf(TEXT("oplog-import %s %s --clean --builds %s --namespace %s --bucket %s --builds-id %s --oidctoken-exe-path \"%s\""),
+				CommandLineBuilder.Append(FString::Printf(TEXT("oplog-import %s %s --builds %s --namespace %s --bucket %s --builds-id %s --oidctoken-exe-path \"%s\""),
 					*DestinationProjectId, *DestinationOplogId, *WriteToString<64>(EffectiveDomain),
 					*State->Namespace, *State->Bucket,
-					*WriteToString<64>(State->BuildId), *OidcTokenExecutableFilename);
+					*WriteToString<64>(State->BuildId), *OidcTokenExecutableFilename));
+
+				if (!State->HostOverride.IsEmpty())
+				{
+					CommandLineBuilder.Appendf(TEXT(" --builds-override-host \"%s\""), *State->HostOverride);
+				}
 			}
 			break;
 		}
 
-		UE_LOG(LogBuildServiceInstance, Display, TEXT("Running: \"%s\" %s."), *ZenUtilityPath, *CommandLineArgs);
-		FMonitoredProcess MonitoredUtilityProcess(ZenUtilityPath, *CommandLineArgs, FPaths::GetPath(ZenUtilityPath), true);
+		if (EnumHasAnyFlags(State->RequestFlags, EBuildTransferRequestFlags::Clean))
+		{
+			CommandLineBuilder.Append(TEXT(" --clean"));
+		}
+
+		if (EnumHasAnyFlags(State->RequestFlags, EBuildTransferRequestFlags::Force))
+		{
+			CommandLineBuilder.Append(TEXT(" --force"));
+		}
+
+		if (EnumHasAllFlags(State->RequestFlags, EBuildTransferRequestFlags::BoostWorkers))
+		{
+			CommandLineBuilder.Append(TEXT(" --boost-workers"));
+		}
+		else if (EnumHasAllFlags(State->RequestFlags, EBuildTransferRequestFlags::BoostWorkerCount))
+		{
+			CommandLineBuilder.Append(TEXT(" --boost-worker-count"));
+		}
+		else if (EnumHasAllFlags(State->RequestFlags, EBuildTransferRequestFlags::BoostWorkerMemory))
+		{
+			CommandLineBuilder.Append(TEXT(" --boost-worker-memory"));
+		}
+
+		UE_LOG(LogBuildServiceInstance, Display, TEXT("Running: \"%s\" %s."), *ZenUtilityPath, *CommandLineBuilder);
+		FMonitoredProcess MonitoredUtilityProcess(ZenUtilityPath, *CommandLineBuilder, FPaths::GetPath(ZenUtilityPath), true);
 		MonitoredUtilityProcess.OnOutput().BindSPLambda(State, [State](FString Output)
 		{
 			FReadScopeLock _(State->Lock);
@@ -1326,7 +1425,7 @@ FBuildServiceInstance::BuildTransferThreadLoop()
 			FString LogFilename = State->GetLogFilename();
 			FOutputDeviceFile::CreateBackupCopy(*LogFilename);
 			State->OutputLogFile = TUniquePtr<FArchive>(IFileManager::Get().CreateFileWriter(*LogFilename, FILEWRITE_Silent | FILEWRITE_AllowRead));
-			State->OutputLogFile->Logf(TEXT("Running: \"%s\" %s"), *ZenUtilityPath, *CommandLineArgs);
+			State->OutputLogFile->Logf(TEXT("Running: \"%s\" %s"), *ZenUtilityPath, *CommandLineBuilder);
 			State->OutputLogFile->Flush();
 			State->Status = EBuildTransferStatus::Active;
 			State->TransferStartTime = FMonotonicTimePoint::Now();
@@ -1543,6 +1642,9 @@ FBuildServiceInstance::RecordTransferStateEvent(TSharedRef<Private::FBuildTransf
 		Attributes.Emplace(TEXT("Bucket"), State->Bucket);
 		Attributes.Emplace(TEXT("ReturnCode"), State->ReturnCode);
 		Attributes.Emplace(TEXT("Platform"), State->TargetPlatformName);
+		Attributes.Emplace(TEXT("HostOverride"), State->HostOverride);
+		Attributes.Emplace(TEXT("UsedDownloadSpec"), State->DownloadSpecJSONContents.IsEmpty() ? false : true);
+		Attributes.Emplace(TEXT("RequestFlags"), static_cast<uint32>(State->RequestFlags));
 		Attributes.Emplace(TEXT("Duration"), FPlatformTime::Seconds() - State->StartTime);
 
 		FStudioTelemetry::Get().RecordEvent(TEXT("Core.BuildService.Transfer"), Attributes);
diff --git a/Engine/Source/Developer/Zen/Public/Experimental/BuildServerInterface.h b/Engine/Source/Developer/Zen/Public/Experimental/BuildServerInterface.h
--- a/Engine/Source/Developer/Zen/Public/Experimental/BuildServerInterface.h
+++ b/Engine/Source/Developer/Zen/Public/Experimental/BuildServerInterface.h
@@ -48,9 +48,23 @@ private:
 	FString OAuthProviderIdentifier;
 };
 
+enum class EBuildTransferRequestFlags
+{
+	None = 0,
+	Clean = 1 << 0,
+	Force = 1 << 1,
+	Scavenge = 1 << 2,
+	Verify = 1 << 3,
+	BoostWorkerCount = 1 << 4,
+	BoostWorkerMemory = 1 << 5,
+	BoostWorkers = BoostWorkerCount | BoostWorkerMemory,
+};
+ENUM_CLASS_FLAGS(EBuildTransferRequestFlags);
+
 class FBuildServiceInstance
 {
 public:
+
 	enum class EBuildTransferStatus
 	{
 		Invalid,
@@ -128,6 +142,11 @@ public:
 
 	UE_API virtual ~FBuildServiceInstance();
 
+	const FServiceSettings& GetSettings() const
+	{
+		return Settings;
+	}
+
 	UE_API void Connect(bool bAllowInteractive, FOnConnectionComplete&& OnConnectionComplete = {});
 	UE_API void RefreshNamespacesAndBuckets(FOnRefreshNamespacesAndBucketsComplete&& InOnRefreshNamespacesAndBucketsComplete = {});
 	[[nodiscard]] UE_API EConnectionState GetConnectionState() const;
@@ -147,7 +166,12 @@ public:
 		FStringView Namespace,
 		FStringView Bucket,
 		FStringView IncludeFilter = FStringView(),
-		FStringView ExcludeFilter = FStringView()
+		FStringView ExcludeFilter = FStringView(),
+		FString&& DownloadSpecJSONContents = FString(),
+		TArray<FString>&& PartIds = TArray<FString>(),
+		TArray<FString>&& PartNames = TArray<FString>(),
+		FStringView HostOverride = FStringView(),
+		EBuildTransferRequestFlags RequestFlags = EBuildTransferRequestFlags::Scavenge
 		);
 	UE_API FBuildTransfer StartOplogBuildTransfer(
 		const FCbObjectId& BuildId,
@@ -159,7 +183,9 @@ public:
 		FStringView ProjectFilePath,
 		FStringView Namespace,
 		FStringView Bucket,
-		FName TargetPlatformName
+		FName TargetPlatformName,
+		FStringView HostOverride = FStringView(),
+		EBuildTransferRequestFlags RequestFlags = EBuildTransferRequestFlags::Clean
 		);
 	UE_API FBuildTransfer RepeatBuildTransfer(FBuildTransfer ExistingTransfer);
 private:
diff --git a/Engine/Source/Editor/StorageServerWidgets/Internal/SBuildSelection.h b/Engine/Source/Editor/StorageServerWidgets/Internal/SBuildSelection.h
--- a/Engine/Source/Editor/StorageServerWidgets/Internal/SBuildSelection.h
+++ b/Engine/Source/Editor/StorageServerWidgets/Internal/SBuildSelection.h
@@ -28,6 +28,10 @@ namespace UE::Zen::Build
 	struct FListBuildsState;
 }
 
+DECLARE_DELEGATE_TwoParams(FPreBuildTransfer, FString& /*HostOverride*/, UE::Zen::Build::EBuildTransferRequestFlags& /*RequestFlags*/);
+
+DECLARE_DELEGATE_TwoParams(FPreOplogBuildTransfer, FString& /*HostOverride*/, UE::Zen::Build::EBuildTransferRequestFlags& /*RequestFlags*/);
+
 DECLARE_DELEGATE_ThreeParams(FOnBuildTransferStarted, UE::Zen::Build::FBuildServiceInstance::FBuildTransfer /*Transfer*/,
 	FStringView /*Name*/,
 	FStringView /*Platform*/);
@@ -42,6 +46,8 @@ class SBuildSelection : public SCompoundWidget
 	SLATE_ATTRIBUTE(TSharedPtr<UE::Zen::FZenServiceInstance>, ZenServiceInstance);
 	SLATE_ATTRIBUTE(TSharedPtr<UE::Zen::Build::FBuildServiceInstance>, BuildServiceInstance);
 
+	SLATE_EVENT(FPreBuildTransfer, PreBuildTransfer)
+	SLATE_EVENT(FPreOplogBuildTransfer, PreOplogBuildTransfer)
 	SLATE_EVENT(FOnBuildTransferStarted, OnBuildTransferStarted)
 	SLATE_END_ARGS()
 
@@ -188,6 +194,7 @@ private:
 	FString UserSelectedDestinations[(int)EBuildType::Count];
 	FTransferFilter UserSelectedTransferFilters[(int)EBuildType::Count];
 	bool bAppendBuildNameToDestinations[(int)EBuildType::Count];
+	bool bCleanDestinations[(int)EBuildType::Count];
 
 	TSharedPtr<FString> SelectedStream;
 	TSharedPtr<FString> SelectedProject;
@@ -208,6 +215,8 @@ private:
 
 	TAttribute<TSharedPtr<UE::Zen::FZenServiceInstance>> ZenServiceInstance;
 	TAttribute<TSharedPtr<UE::Zen::Build::FBuildServiceInstance>> BuildServiceInstance;
+	FPreBuildTransfer PreBuildTransfer;
+	FPreOplogBuildTransfer PreOplogBuildTransfer;
 	FOnBuildTransferStarted OnBuildTransferStarted;
 };
 
diff --git a/Engine/Source/Editor/StorageServerWidgets/Private/SBuildSelection.cpp b/Engine/Source/Editor/StorageServerWidgets/Private/SBuildSelection.cpp
--- a/Engine/Source/Editor/StorageServerWidgets/Private/SBuildSelection.cpp
+++ b/Engine/Source/Editor/StorageServerWidgets/Private/SBuildSelection.cpp
@@ -132,6 +132,8 @@ void SBuildSelection::Construct(const FArguments& InArgs)
 
 	ZenServiceInstance = InArgs._ZenServiceInstance;
 	BuildServiceInstance = InArgs._BuildServiceInstance;
+	PreBuildTransfer = InArgs._PreBuildTransfer;
+	PreOplogBuildTransfer = InArgs._PreOplogBuildTransfer;
 	OnBuildTransferStarted = InArgs._OnBuildTransferStarted;
 
 	BuildGroupSortByColumn = FBuildGroupIds::ColCreated;
@@ -1662,8 +1664,25 @@ void SBuildSelection::DownloadSelectedArtifactsInSelectedBuildGroup()
 					{
 						DestinationOplogId = FString::Printf(TEXT("%s.%s"), *DestinationOplogId, *SanitizeForZenId(SelectedItems[0]->DisplayName));
 					}
+
+					FString HostOverride;
+					EBuildTransferRequestFlags RequestFlags = EBuildTransferRequestFlags::Clean;
+					PreOplogBuildTransfer.ExecuteIfBound(HostOverride, RequestFlags);
+
 					FBuildServiceInstance::FBuildTransfer BuildTransfer =
-						ServiceInstance->StartOplogBuildTransfer(Artifact->BuildRecord.BuildId, TransferName, DestinationProjectId, DestinationOplogId, RootDir, EngineDir, ProjectFilename, SelectedItems[0]->Namespace, Artifact->BuildRecord.BucketId, TargetPlatformName);
+						ServiceInstance->StartOplogBuildTransfer(
+							Artifact->BuildRecord.BuildId,
+							TransferName,
+							DestinationProjectId,
+							DestinationOplogId,
+							RootDir,
+							EngineDir,
+							ProjectFilename,
+							SelectedItems[0]->Namespace,
+							Artifact->BuildRecord.BucketId,
+							TargetPlatformName,
+							HostOverride,
+							RequestFlags);
 					OnBuildTransferStarted.ExecuteIfBound(BuildTransfer, SelectedItems[0]->DisplayName, ArtifactForSelectedGroup);
 				}
 				else
@@ -1677,6 +1696,11 @@ void SBuildSelection::DownloadSelectedArtifactsInSelectedBuildGroup()
 					{
 						DestinationFolder = FPaths::Combine(GetEffectiveDestination(), DestinationPlatformName);
 					}
+
+					FString HostOverride;
+					EBuildTransferRequestFlags RequestFlags = EBuildTransferRequestFlags::Scavenge;
+					PreBuildTransfer.ExecuteIfBound(HostOverride, RequestFlags);
+
 					FBuildServiceInstance::FBuildTransfer BuildTransfer =
 						ServiceInstance->StartBuildTransfer(
 							Artifact->BuildRecord.BuildId,
@@ -1685,7 +1709,12 @@ void SBuildSelection::DownloadSelectedArtifactsInSelectedBuildGroup()
 							SelectedItems[0]->Namespace,
 							Artifact->BuildRecord.BucketId,
 							UserSelectedTransferFilters[BuildTypeIndex].IncludeFilter,
-							UserSelectedTransferFilters[BuildTypeIndex].ExcludeFilter);
+							UserSelectedTransferFilters[BuildTypeIndex].ExcludeFilter,
+							FString(),
+							TArray<FString>(),
+							TArray<FString>(),
+							HostOverride,
+							RequestFlags);
 					OnBuildTransferStarted.ExecuteIfBound(BuildTransfer, SelectedItems[0]->DisplayName, ArtifactForSelectedGroup);
 				}
 			}
diff --git a/Engine/Source/Programs/BuildStorageTool/Private/BuildStorageTool.cpp b/Engine/Source/Programs/BuildStorageTool/Private/BuildStorageTool.cpp
--- a/Engine/Source/Programs/BuildStorageTool/Private/BuildStorageTool.cpp
+++ b/Engine/Source/Programs/BuildStorageTool/Private/BuildStorageTool.cpp
@@ -38,6 +38,7 @@
 #include "OutputLogModule.h"
 #include "OutputLogSettings.h"
 #include "StudioTelemetry.h"
+#include "String/ParseTokens.h"
 
 
 #if PLATFORM_WINDOWS
@@ -77,16 +78,6 @@ static void HideOnCloseOverride(const TSharedRef<SWindow>& WindowBeingClosed)
 
 class FBuildStorageToolApp
 {
-	FCriticalSection CriticalSection;
-	FSlateApplication& Slate;
-	TSharedPtr<SWindow> Window;
-	TSharedPtr<SNotificationItem> CompileNotification;
-	TSharedPtr<UE::Zen::FServiceInstanceManager> ZenServiceInstanceManager;
-	TSharedPtr<UE::Zen::Build::FServiceInstanceManager> BuildServiceInstanceManager;
-	const FBuildStorageToolParameters& ToolParameters;
-
-	std::atomic<bool> bLatentExclusiveOperationActive = false;
-
 	void ExitTool()
 	{
 		FSlateApplication::Get().RequestDestroyWindow(Window.ToSharedRef());
@@ -244,10 +235,135 @@ class FBuildStorageToolApp
 					NAME_None,
 					EUserInterfaceActionType::Button
 				);
-		
+
+		MenuBuilder.BeginSection(TEXT("TransferOptions"), LOCTEXT("TransferOptionsMenu", "Transfer Options"));
+		{
+			MenuBuilder.AddMenuEntry(
+				LOCTEXT("CleanDestination", "Clean Extra Files"),
+				LOCTEXT("CleanDestination_ToolTip", "When downloading files to a directory, any extra files that existed in that directory will be deleted"),
+				FSlateIcon(),
+				FUIAction(
+					FExecuteAction::CreateRaw(
+						this, &FBuildStorageToolApp::ToggleCleanDestination),
+					FCanExecuteAction(),
+					FIsActionChecked::CreateRaw(
+						this, &FBuildStorageToolApp::IsCleanDestinationSelected)
+				),
+				NAME_None,
+				EUserInterfaceActionType::Check
+			);
+
+			MenuBuilder.AddMenuEntry(
+				LOCTEXT("ForceRedownload", "Force Redownload"),
+				LOCTEXT("ForceRedownload_ToolTip", "Forces downloads to get all data not just the data that is not already there"),
+				FSlateIcon(),
+				FUIAction(
+					FExecuteAction::CreateRaw(
+						this, &FBuildStorageToolApp::ToggleForceRedownload),
+					FCanExecuteAction(),
+					FIsActionChecked::CreateRaw(
+						this, &FBuildStorageToolApp::IsForceRedownloadSelected)
+				),
+				NAME_None,
+				EUserInterfaceActionType::Check
+			);
+
+			MenuBuilder.AddMenuEntry(
+				LOCTEXT("EnableScavenging", "Enable Scavenging"),
+				LOCTEXT("EnableScavenging_ToolTip", "Allows downloads to find and use data from downloads elsewhere on your machine"),
+				FSlateIcon(),
+				FUIAction(
+					FExecuteAction::CreateRaw(
+						this, &FBuildStorageToolApp::ToggleScavenging),
+					FCanExecuteAction::CreateRaw(
+						this, &FBuildStorageToolApp::CanSelectScavenging),
+					FIsActionChecked::CreateRaw(
+						this, &FBuildStorageToolApp::IsScavengingSelected)
+				),
+				NAME_None,
+				EUserInterfaceActionType::Check
+			);
+
+			MenuBuilder.AddMenuEntry(
+				LOCTEXT("EnableVerify", "Enable Verification"),
+				LOCTEXT("EnableVerify_ToolTip", "Performs a verification pass on data after it has been downloaded"),
+				FSlateIcon(),
+				FUIAction(
+					FExecuteAction::CreateRaw(
+						this, &FBuildStorageToolApp::ToggleVerify),
+					FCanExecuteAction(),
+					FIsActionChecked::CreateRaw(
+						this, &FBuildStorageToolApp::IsVerifySelected)
+				),
+				NAME_None,
+				EUserInterfaceActionType::Check
+			);
+
+			MenuBuilder.AddMenuEntry(
+				LOCTEXT("BoostWorkers", "Boost Downloads"),
+				LOCTEXT("BoostWorkers_ToolTip", "Increases resource usage when downloading, computer may be less responsive"),
+				FSlateIcon(),
+				FUIAction(
+					FExecuteAction::CreateRaw(
+						this, &FBuildStorageToolApp::ToggleBoostWorkers),
+					FCanExecuteAction(),
+					FIsActionChecked::CreateRaw(
+						this, &FBuildStorageToolApp::IsBoostWorkersSelected)
+				),
+				NAME_None,
+				EUserInterfaceActionType::Check
+			);
+
+			MenuBuilder.AddSubMenu(
+				LOCTEXT("HostOverride", "Host Override"),
+				LOCTEXT("HostOverride_ToolTip", "Pick a host to use for build transfers instead of automatic selection of the nearest host"),
+				FNewMenuDelegate::CreateRaw(this, &FBuildStorageToolApp::OpenHostOverrideMenu));
+		}
+		MenuBuilder.EndSection();
 		return MenuBuilder.MakeWidget();
 	}
 
+	void OpenHostOverrideMenu(FMenuBuilder& MenuBuilder)
+	{
+		using namespace UE::Zen::Build;
+		const FServiceSettings& ServiceSettings = BuildServiceInstanceManager->GetBuildServiceInstance()->GetSettings();
+
+		MenuBuilder.AddMenuEntry(
+			LOCTEXT("NoneHostOverride", "None"),
+			LOCTEXT("NoneHostOverride_ToolTip", "Disables host overriding and restores automatic selection of nearest host when downloading"),
+			FSlateIcon(),
+			FUIAction(
+				FExecuteAction::CreateRaw(this, &FBuildStorageToolApp::SetHostOverride, FString()),
+				FCanExecuteAction(),
+				FIsActionChecked::CreateRaw(this, &FBuildStorageToolApp::CompareHostOverride, FString())
+			),
+			NAME_None,
+			EUserInterfaceActionType::RadioButton
+		);
+
+		TArray<FStringView> HostCandidates;
+		UE::String::ParseTokens(ServiceSettings.GetHost(), ';', [&HostCandidates](FStringView Host)
+		{
+			HostCandidates.Emplace(Host);
+		}, UE::String::EParseTokensOptions::SkipEmpty | UE::String::EParseTokensOptions::Trim);
+
+		for (FStringView HostCandidate : HostCandidates)
+		{
+			MenuBuilder.AddMenuEntry(
+				FText::FromStringView(HostCandidate),
+				FText(),
+				FSlateIcon(),
+				FUIAction(
+					FExecuteAction::CreateRaw(this, &FBuildStorageToolApp::SetHostOverride, FString(HostCandidate)),
+					FCanExecuteAction(),
+					FIsActionChecked::CreateRaw(this, &FBuildStorageToolApp::CompareHostOverride, FString(HostCandidate))
+				),
+				NAME_None,
+				EUserInterfaceActionType::RadioButton
+			);
+		}
+	}
+
 	TSharedRef<SWidget> FillFileMenu()
 	{
 		const bool bCloseSelfOnly = false;
@@ -371,6 +487,48 @@ public:
 						SNew(SBuildSelection)
 						.ZenServiceInstance(ZenServiceInstanceManager.ToSharedRef(), &UE::Zen::FServiceInstanceManager::GetZenServiceInstance)
 						.BuildServiceInstance(BuildServiceInstanceManager.ToSharedRef(), &UE::Zen::Build::FServiceInstanceManager::GetBuildServiceInstance)
+						.PreBuildTransfer_Lambda([this]
+							(FString& InHostOverride, UE::Zen::Build::EBuildTransferRequestFlags& RequestFlags)
+							{
+								InHostOverride = HostOverride;
+								RequestFlags = UE::Zen::Build::EBuildTransferRequestFlags::None;
+								if (bCleanDestination)
+								{
+									RequestFlags |= UE::Zen::Build::EBuildTransferRequestFlags::Clean;
+								}
+								if (bForceRedownload)
+								{
+									RequestFlags |= UE::Zen::Build::EBuildTransferRequestFlags::Force;
+								}
+								else if (bScavenging)
+								{
+									RequestFlags |= UE::Zen::Build::EBuildTransferRequestFlags::Scavenge;
+								}
+								if (bVerify)
+								{
+									RequestFlags |= UE::Zen::Build::EBuildTransferRequestFlags::Verify;
+								}
+								if (bBoostWorkers)
+								{
+									RequestFlags |= UE::Zen::Build::EBuildTransferRequestFlags::BoostWorkers;
+								}
+							})
+						.PreOplogBuildTransfer_Lambda([this]
+							(FString& InHostOverride, UE::Zen::Build::EBuildTransferRequestFlags& RequestFlags)
+							{
+								InHostOverride = HostOverride;
+								RequestFlags = UE::Zen::Build::EBuildTransferRequestFlags::None;
+								// We want oplogs to always download "clean"
+								RequestFlags |= UE::Zen::Build::EBuildTransferRequestFlags::Clean;
+								if (bForceRedownload)
+								{
+									RequestFlags |= UE::Zen::Build::EBuildTransferRequestFlags::Force;
+								}
+								if (bBoostWorkers)
+								{
+									RequestFlags |= UE::Zen::Build::EBuildTransferRequestFlags::BoostWorkers;
+								}
+							})
 						.OnBuildTransferStarted_Lambda([this]
 							(UE::Zen::Build::FBuildServiceInstance::FBuildTransfer Transfer, FStringView Name, FStringView Platform)
 							{
@@ -461,6 +619,86 @@ private:
 		FCoreDelegates::ModalMessageDialog.Unbind();
 	}
 
+	void ToggleCleanDestination()
+	{
+		bCleanDestination = !bCleanDestination;
+	}
+
+	bool IsCleanDestinationSelected()
+	{
+		return bCleanDestination;
+	}
+
+	void ToggleForceRedownload()
+	{
+		bForceRedownload = !bForceRedownload;
+	}
+
+	bool IsForceRedownloadSelected()
+	{
+		return bForceRedownload;
+	}
+
+	void ToggleScavenging()
+	{
+		bScavenging = !bScavenging;
+	}
+
+	bool CanSelectScavenging()
+	{
+		return !bForceRedownload;
+	}
+
+	bool IsScavengingSelected()
+	{
+		return bScavenging && !bForceRedownload;
+	}
+
+	void ToggleVerify()
+	{
+		bVerify = !bVerify;
+	}
+
+	bool IsVerifySelected()
+	{
+		return bVerify;
+	}
+
+	void ToggleBoostWorkers()
+	{
+		bBoostWorkers = !bBoostWorkers;
+	}
+
+	bool IsBoostWorkersSelected()
+	{
+		return bBoostWorkers;
+	}
+
+	void SetHostOverride(FString InHostOverride)
+	{
+		HostOverride = InHostOverride;
+	}
+
+	bool CompareHostOverride(FString InHostOverride) const
+	{
+		return HostOverride == InHostOverride;
+	}
+
+	FCriticalSection CriticalSection;
+	FSlateApplication& Slate;
+	TSharedPtr<SWindow> Window;
+	TSharedPtr<SNotificationItem> CompileNotification;
+	TSharedPtr<UE::Zen::FServiceInstanceManager> ZenServiceInstanceManager;
+	TSharedPtr<UE::Zen::Build::FServiceInstanceManager> BuildServiceInstanceManager;
+	const FBuildStorageToolParameters& ToolParameters;
+	FString HostOverride;
+	bool bCleanDestination = false;
+	bool bForceRedownload = false;
+	bool bScavenging = true;
+	bool bVerify = false;
+	bool bBoostWorkers = false;
+
+	std::atomic<bool> bLatentExclusiveOperationActive = false;
 	TSharedPtr<SBuildActivity> BuildActivity;
 };
 
