Increase curl upload buffer size. It improves the large file upload speed, such as zenfs.manifest, on Linux cloud workers drastically.

#Synced-CL 51332327

diff --git a/Engine/Source/Developer/Zen/Private/ZenServerHttp.cpp b/Engine/Source/Developer/Zen/Private/ZenServerHttp.cpp
--- a/Engine/Source/Developer/Zen/Private/ZenServerHttp.cpp
+++ b/Engine/Source/Developer/Zen/Private/ZenServerHttp.cpp
@@ -103,7 +103,7 @@ namespace UE::Zen {
 		curl_easy_setopt(Curl, CURLOPT_DNS_CACHE_TIMEOUT, -1L); // Don't re-resolve names mid-session
 		curl_easy_setopt(Curl, CURLOPT_BUFFERSIZE, 256 * 1024L);
 		curl_easy_setopt(Curl, CURLOPT_NOPROXY, "*");
-		//curl_easy_setopt(Curl, CURLOPT_UPLOAD_BUFFERSIZE, 256 * 1024L);
+		curl_easy_setopt(Curl, CURLOPT_UPLOAD_BUFFERSIZE, 2 * 1024L * 1024L); // Max buffer size is 2MB. 
 		// Response functions
 		curl_easy_setopt(Curl, CURLOPT_HEADERDATA, this);
 		curl_easy_setopt(Curl, CURLOPT_HEADERFUNCTION, &FZenHttpRequest::FStatics::StaticWriteHeaderFn);
