Downgrade message verbosity when a zen plugin file can't be found.  This happens when users have filtered out platform support.

#Synced-CL 51529686

diff --git a/Engine/Source/Developer/Zen/Private/ZenServerInterface.cpp b/Engine/Source/Developer/Zen/Private/ZenServerInterface.cpp
--- a/Engine/Source/Developer/Zen/Private/ZenServerInterface.cpp
+++ b/Engine/Source/Developer/Zen/Private/ZenServerInterface.cpp
@@ -1132,7 +1132,7 @@ bool FServicePluginSettings::ReadFromConfig(const FString& InPluginName)
 			}
 			else
 			{
-				UE_LOG(LogZenServiceInstance, Warning, TEXT("Can't find file for plugin '%s', tried '%s'"), *Name, *AbsPathRoot);
+				UE_LOG(LogZenServiceInstance, Display, TEXT("Can't find file for plugin '%s', tried '%s'"), *Name, *AbsPathRoot);
 				return false;
 			}
 
