Add support for listing build contents to the build service interface.

#Synced-CL 51529686

diff --git a/Engine/Source/Developer/Zen/Private/BuildServerInterface.cpp b/Engine/Source/Developer/Zen/Private/BuildServerInterface.cpp
--- a/Engine/Source/Developer/Zen/Private/BuildServerInterface.cpp
+++ b/Engine/Source/Developer/Zen/Private/BuildServerInterface.cpp
@@ -538,6 +538,85 @@ FBuildServiceInstance::ListBuildsAcrossBuckets(FStringView Namespace, FStringVie
 	});
 }
 
+void
+FBuildServiceInstance::ListBuildContents(FStringView Namespace, FStringView Bucket, FStringView BuildId, FOnListBuildContentsComplete&& InOnListBuildContentsComplete)
+{
+	if (ConnectionState.load(std::memory_order_relaxed) != EConnectionState::ConnectionSucceeded)
+	{
+		if (InOnListBuildContentsComplete)
+		{
+			InOnListBuildContentsComplete({});
+		}
+		return;
+	}
+
+	FString OutputFilePath = FPaths::CreateTempFilename(FPlatformProcess::UserTempDir(), TEXT("BuildContents"), TEXT(".cbo"));
+	FPaths::MakePlatformFilename(OutputFilePath);
+
+	FString OidcTokenExecutableFilename = FPaths::ConvertRelativePathToFull(FDesktopPlatformModule::TryGet()->GetOidcTokenExecutableFilename(FPaths::RootDir()));
+
+	FString CommandLineArgs = FString::Printf(TEXT("builds ls --namespace %s --bucket %s --build-id %s --result-path \"%s\""),
+		*WriteToString<64>(Namespace),
+		*WriteToString<64>(Bucket),
+		*WriteToString<24>(BuildId),
+		*OutputFilePath);
+	CommandLineArgs = AddZenUtilityBuildServerArguments(*CommandLineArgs);
+
+	InvokeZenUtility(CommandLineArgs, nullptr,
+	[OutputFilePath = MoveTemp(OutputFilePath), InOnListBuildContentsComplete = MoveTemp(InOnListBuildContentsComplete)](bool bSuccessful) mutable
+	{
+		ON_SCOPE_EXIT
+		{
+			IFileManager::Get().Delete(*OutputFilePath, false, false, true);
+		};
+
+		if (!bSuccessful)
+		{
+			if (InOnListBuildContentsComplete)
+			{
+				InOnListBuildContentsComplete({});
+			}
+			return;
+		}
+
+		TUniquePtr<FArchive> Ar(IFileManager::Get().CreateFileReader(*OutputFilePath));
+		if (!Ar)
+		{
+			UE_LOG(LogBuildServiceInstance, Warning, TEXT("Missing output file from zen utility when listing build contents: '%s'"), *OutputFilePath);
+			if (InOnListBuildContentsComplete)
+			{
+				InOnListBuildContentsComplete({});
+			}
+			return;
+		}
+
+		if (InOnListBuildContentsComplete)
+		{
+			FCbObject OutputObject = LoadCompactBinary(*Ar).AsObject();
+			TMap<FString, FBuildPart> BuildContents;
+			for (FCbFieldView PartField : OutputObject["parts"])
+			{
+				FString PartName = *WriteToString<64>(PartField["partName"].AsString());
+
+				FBuildPart Part;
+				Part.Id = PartField["id"].AsObjectId();
+
+				for (FCbFieldView FileField : PartField["files"])
+				{
+					FBuildFile File;
+					File.Name = *WriteToString<64>(FileField["path"].AsString());
+					File.RawSize = FileField["rawSize"].AsUInt64();
+					File.Attributes = static_cast<uint32>(FileField["attributes"].AsUInt64());
+					Part.Files.Emplace(MoveTemp(File));
+				}
+
+				BuildContents.Add(MoveTemp(PartName), MoveTemp(Part));
+			}
+			InOnListBuildContentsComplete(MoveTemp(BuildContents));
+		}
+	});
+}
+
 FBuildServiceInstance::FBuildTransfer
 FBuildServiceInstance::StartBuildTransfer(
 	const FCbObjectId& BuildId,
diff --git a/Engine/Source/Developer/Zen/Public/Experimental/BuildServerInterface.h b/Engine/Source/Developer/Zen/Public/Experimental/BuildServerInterface.h
--- a/Engine/Source/Developer/Zen/Public/Experimental/BuildServerInterface.h
+++ b/Engine/Source/Developer/Zen/Public/Experimental/BuildServerInterface.h
@@ -104,6 +104,19 @@ public:
 		UE_API FString GetCookPlatform() const;
 	};
 
+	struct FBuildFile
+	{
+		FString Name;
+		uint64 RawSize = 0;
+		uint32 Attributes = 0;
+	};
+
+	struct FBuildPart
+	{
+		FCbObjectId Id;
+		TArray<FBuildFile> Files;
+	};
+
 	struct FBuildTransfer
 	{
 	public:
@@ -136,6 +149,7 @@ public:
 	typedef TUniqueFunction<void(EConnectionState, EConnectionFailureReason)> FOnConnectionComplete;
 	typedef TUniqueFunction<void()> FOnRefreshNamespacesAndBucketsComplete;
 	typedef TUniqueFunction<void(TArray<FBuildRecord>&& ListBuildRecords)> FOnListBuildsComplete;
+	typedef TUniqueFunction<void(TMap<FString, FBuildPart>&& BuildContents)> FOnListBuildContentsComplete;
 
 
 	UE_API FBuildServiceInstance();
@@ -159,6 +173,7 @@ public:
 
 	UE_API void ListBuilds(FStringView Namespace, FStringView Bucket, FCbObject&& Query, FOnListBuildsComplete&& InOnListBuildsComplete);
 	UE_API void ListBuildsAcrossBuckets(FStringView Namespace, FStringView BucketRegex, FCbObject&& Query, FOnListBuildsComplete&& InOnListBuildsComplete);
+	UE_API void ListBuildContents(FStringView Namespace, FStringView Bucket, FStringView BuildId, FOnListBuildContentsComplete&& InOnListBuildContentsComplete);
 
 	UE_API FBuildTransfer StartBuildTransfer(
 		const FCbObjectId& BuildId,
