Update Zen installation to use well-known OidcToken executable paths and restrict content types on certain requests.

[RN] none
#redirectsaresafe

#Synced-CL 51771667

#Synced-CL 51771667

diff --git a/Engine/Config/BaseEngine.ini b/Engine/Config/BaseEngine.ini
--- a/Engine/Config/BaseEngine.ini
+++ b/Engine/Config/BaseEngine.ini
@@ -2688,7 +2688,7 @@ AllowRemoteNetworkService=false
 LocalDataCachePathEnvOverride=UE-LocalDataCachePath
 LocalDataCachePathCommandLineOverride=LocalDataCachePath
 LocalDataCachePathEditorOverrideSetting=LocalDerivedDataCache
-ExtraArgs=--http asio --gc-cache-duration-seconds 1209600 --gc-interval-seconds 21600 --gc-low-diskspace-threshold 2147483648 --cache-bucket-limit-overwrites --quiet
+ExtraArgs=--http asio --gc-cache-duration-seconds 1209600 --gc-interval-seconds 21600 --gc-low-diskspace-threshold 2147483648 --cache-bucket-limit-overwrites --quiet --restrict-content-types --allow-external-oidctoken-exe=false
 Plugins=AndroidPortForwarder
 
 [Zen.ConnectExisting]
diff --git a/Engine/Source/Developer/Zen/Private/ZenServerInterface.cpp b/Engine/Source/Developer/Zen/Private/ZenServerInterface.cpp
--- a/Engine/Source/Developer/Zen/Private/ZenServerInterface.cpp
+++ b/Engine/Source/Developer/Zen/Private/ZenServerInterface.cpp
@@ -186,6 +186,18 @@ GetUtilityExecutableName()
 		;
 }
 
+static FString
+GetOidcTokenExecutableName()
+{
+	return
+#if PLATFORM_WINDOWS
+		TEXT("OidcToken.exe")
+#else
+		TEXT("OidcToken")
+#endif
+		;
+}
+
 static FString
 GetSystemInstallPath(const FSystemServiceSettings& InSettings)
 {
@@ -335,6 +347,28 @@ GetInstallCrashpadHandlerFilePath(const FString& InTreePath)
 	return FPaths::ConvertRelativePathToFull(FPaths::Combine(GetLocalInstallPath(), FString(FPathViews::GetCleanFilename(InTreePath))));
 }
 
+static FString
+GetInTreeOidcTokenPath()
+{
+	const FString PlatformDir =
+#if PLATFORM_WINDOWS
+		TEXT("win-x64")
+#elif PLATFORM_MAC
+		TEXT("osx-x64")
+#else	
+		TEXT("linux-x64")
+#endif
+	;
+
+	return FPaths::ConvertRelativePathToFull(FPaths::Combine(FPaths::EngineDir(), TEXT("Binaries"), TEXT("DotNet"), TEXT("OidcToken"), PlatformDir, GetOidcTokenExecutableName()));
+}
+
+static FString
+GetInstallOidcTokenPath()
+{
+	return FPaths::ConvertRelativePathToFull(FPaths::Combine(GetLocalInstallPath(), GetOidcTokenExecutableName()));
+}
+
 #if PLATFORM_WINDOWS
 
 FString WriteMiniDump(HANDLE ProcessHandle, const FString& ExeName, const FString& OutputFolderPath)
@@ -3781,6 +3815,17 @@ FZenServiceInstance::ConditionalUpdateLocalInstall(FServiceAutoLaunchSettings::E
 			AttemptFileCopyWithRetries(*InstallCrashpadHandlerFilePath, *InTreeCrashpadHandlerFilePath, 1.0);
 		}
 
+		FString InTreeOidcTokenPath = GetInTreeOidcTokenPath();
+		FString InstallOidcTokenPath = GetInstallOidcTokenPath();
+
+		if (FileManager.FileExists(*InTreeOidcTokenPath) && (bMainExecutablesUpdated || !FileManager.FileExists(*InstallOidcTokenPath)))
+		{
+			if (!AttemptFileCopyWithRetries(*InstallOidcTokenPath, *InTreeOidcTokenPath, 1.0))
+			{
+				UE_LOG(LogZenServiceInstance, Warning, TEXT("Failed to update OidcToken executable"));
+			}
+		}
+
 		const FString LinkPath = GetServiceLinkPath();
 		if (FZenServiceLink Link = FZenServiceLink::Read(LinkPath); Link)
 		{
@@ -4002,6 +4047,11 @@ FZenServiceInstance::AutoLaunch(const FServiceAutoLaunchSettings& InSettings, co
 
 		FString ParmsWithoutTransients = DetermineCmdLineWithoutTransientComponents(InSettings.DataPath, InSettings.ExtraArgs, InSettings.Plugins, InSettings.bSendUnattendedBugReports, InSettings.bAllowRemoteNetworkService, InSettings.DesiredPort);
 
+		if (InSettings.InstallMode == FServiceAutoLaunchSettings::EInstallMode::Link)
+		{
+			ParmsWithoutTransients.Appendf(TEXT(" --oidctoken-exe-path \"%s\""), *GetInTreeOidcTokenPath());
+		}
+
 		FZenLocalServiceRunContext EffectiveRunContext;
 		EffectiveRunContext.Executable = ExecutablePath;
 		EffectiveRunContext.CommandlineArguments = ParmsWithoutTransients;
