BuildStorageTool: If HostBuilder host resolving doesn't successfully contact any host, still fall back to attempting to use the default host.

#Synced-CL 51529686

diff --git a/Engine/Source/Developer/Zen/Private/BuildServerInterface.cpp b/Engine/Source/Developer/Zen/Private/BuildServerInterface.cpp
--- a/Engine/Source/Developer/Zen/Private/BuildServerInterface.cpp
+++ b/Engine/Source/Developer/Zen/Private/BuildServerInterface.cpp
@@ -274,45 +274,37 @@ FBuildServiceInstance::Connect(bool bAllowInteractive, FOnConnectionComplete&& O
 		double ResolvedLatency;
 		FHttpHostBuilder HostBuilder;
 		HostBuilder.AddFromString(Settings.GetHost());
-		if (HostBuilder.ResolveHost(/* Warning timeout */ 1.0, 4.0 /* Max duration timeout*/, ResolvedHost, ResolvedLatency))
+		bool bHostResolvingSucceeded = HostBuilder.ResolveHost(/* Warning timeout */ 1.0, 4.0 /* Max duration timeout*/, ResolvedHost, ResolvedLatency);
+		
+		// even if we fail to resolve a host to use the returned host will at least contain the first of the possible hosts which we can attempt to use
+		EffectiveDomain = ResolvedHost;
+
+		if (!bHostResolvingSucceeded)
 		{
-			EffectiveDomain = ResolvedHost;
+			FString HostCandidates = HostBuilder.GetHostCandidatesString();
+			const FString MessageString = FString::Printf(TEXT("Unable to resolve best host candidate to use. Most likely none of the suggested hosts was reachable. Attempted hosts were: '%s' ."), *HostBuilder.GetHostCandidatesString() );
+			UE_LOG(LogBuildServiceInstance, Warning, TEXT("%s"), *MessageString);
+		}
 
-			EDesktopLoginInteractionLevel InteractionLevel = bAllowInteractive ? EDesktopLoginInteractionLevel::Interactive : EDesktopLoginInteractionLevel::None;
-			if (AcquireAccessToken(InteractionLevel))
-			{
-				RecordConnectionEvent(true);
-				ConnectionState.store(EConnectionState::ConnectionSucceeded, std::memory_order_relaxed);
-				if (OnConnectionComplete)
-				{
-					OnConnectionComplete(EConnectionState::ConnectionSucceeded, EConnectionFailureReason::None);
-				}
-			}
-			else
+		EDesktopLoginInteractionLevel InteractionLevel = bAllowInteractive ? EDesktopLoginInteractionLevel::Interactive : EDesktopLoginInteractionLevel::None;
+		if (AcquireAccessToken(InteractionLevel))
+		{
+			RecordConnectionEvent(true);
+			ConnectionState.store(EConnectionState::ConnectionSucceeded, std::memory_order_relaxed);
+			if (OnConnectionComplete)
 			{
-				const FString MessageString = TEXT("Unable to acquire access token.");
-				UE_LOG(LogBuildServiceInstance, Warning, TEXT("%s"), *MessageString);
-				RecordConnectionEvent(false, MessageString);
-				ConnectionState.store(EConnectionState::ConnectionFailed, std::memory_order_relaxed);
-				if (OnConnectionComplete)
-				{
-					OnConnectionComplete(EConnectionState::ConnectionFailed, EConnectionFailureReason::FailedToAcquireAccessToken);
-				}
+				OnConnectionComplete(EConnectionState::ConnectionSucceeded, EConnectionFailureReason::None);
 			}
 		}
 		else
 		{
-			// even if we fail to resolve a host to use the returned host will at least contain the first of the possible hosts which we can attempt to use
-			EffectiveDomain = ResolvedHost;
-
-			FString HostCandidates = HostBuilder.GetHostCandidatesString();
-			const FString MessageString = FString::Printf(TEXT("Unable to resolve best host candidate to use. Most likely none of the suggested hosts was reachable. Attempted hosts were: '%s' ."), *HostCandidates );
+			const FString MessageString = TEXT("Unable to acquire access token.");
 			UE_LOG(LogBuildServiceInstance, Warning, TEXT("%s"), *MessageString);
 			RecordConnectionEvent(false, MessageString);
 			ConnectionState.store(EConnectionState::ConnectionFailed, std::memory_order_relaxed);
 			if (OnConnectionComplete)
 			{
-				OnConnectionComplete(EConnectionState::ConnectionFailed, EConnectionFailureReason::FailedToResolveHost);
+				OnConnectionComplete(EConnectionState::ConnectionFailed, EConnectionFailureReason::FailedToAcquireAccessToken);
 			}
 		}
 
