Add remote network service modes for zenserver autolaunch.  These control security and access for remote machines to a locally launched zenserver.  A new mode is the GeneratedStaticKey mode which allows for a generated key that is used as a password by game clients.  This requires host authentication info to be included in the ue.projectstore file so that a game running on anothe rmachine is able to authenticate appropriately.

#redirectsaresafe
#Synced-CL 52199356

diff --git a/Engine/Config/BaseEngine.ini b/Engine/Config/BaseEngine.ini
--- a/Engine/Config/BaseEngine.ini
+++ b/Engine/Config/BaseEngine.ini
@@ -2695,8 +2695,8 @@ AutoLaunch=true
 
 [Zen.AutoLaunch]
 DataPath=%ENGINEVERSIONAGNOSTICINSTALLEDUSERDIR%Zen/Data
-LimitProcessLifetime=true
-AllowRemoteNetworkService=false
+LimitProcessLifetime=True
+RemoteNetworkService=None
 LocalDataCachePathEnvOverride=UE-LocalDataCachePath
 LocalDataCachePathCommandLineOverride=LocalDataCachePath
 LocalDataCachePathEditorOverrideSetting=LocalDerivedDataCache
diff --git a/Engine/Plugins/Developer/ProjectLauncher/Source/ProjectLauncher/Private/Model/ProjectLauncherModel.cpp b/Engine/Plugins/Developer/ProjectLauncher/Source/ProjectLauncher/Private/Model/ProjectLauncherModel.cpp
--- a/Engine/Plugins/Developer/ProjectLauncher/Source/ProjectLauncher/Private/Model/ProjectLauncherModel.cpp
+++ b/Engine/Plugins/Developer/ProjectLauncher/Source/ProjectLauncher/Private/Model/ProjectLauncherModel.cpp
@@ -989,10 +989,9 @@ PRAGMA_ENABLE_DEPRECATION_WARNINGS
 #if UE_WITH_ZEN
 		using namespace UE::Zen;
 		const FServiceSettings& ZenServiceSettings = GetDefaultServiceInstance().GetServiceSettings();
-		OutResult.bAllowRemoteNetworkService = ZenServiceSettings.IsAutoLaunch() && ZenServiceSettings.SettingsVariant.Get<FServiceAutoLaunchSettings>().bAllowRemoteNetworkService;
+		OutResult.bAllowRemoteNetworkService = ZenServiceSettings.IsAutoLaunch() && ZenServiceSettings.GetRemoteNetworkService() != ERemoteNetworkService::None;
 #else
-		const TCHAR* ZenAutoLaunchSettings = TEXT("Zen.AutoLaunch");
-		InConfig.GetBool(ZenAutoLaunchSettings, TEXT("AllowRemoteNetworkService"), OutResult.bAllowRemoteNetworkService, GEngineIni);
+		OutResult.bAllowRemoteNetworkService = false;
 #endif
 
 	}
diff --git a/Engine/Plugins/Developer/ProjectLauncher/Source/ProjectLauncher/Private/ProfileTree/GenericProfileTreeBuilder.cpp b/Engine/Plugins/Developer/ProjectLauncher/Source/ProjectLauncher/Private/ProfileTree/GenericProfileTreeBuilder.cpp
--- a/Engine/Plugins/Developer/ProjectLauncher/Source/ProjectLauncher/Private/ProfileTree/GenericProfileTreeBuilder.cpp
+++ b/Engine/Plugins/Developer/ProjectLauncher/Source/ProjectLauncher/Private/ProfileTree/GenericProfileTreeBuilder.cpp
@@ -855,7 +855,7 @@ namespace ProjectLauncher
 			
 			if (!Model->GetProjectSettings(Owner.Profile).bAllowRemoteNetworkService && Model->IsUsingRemotePlatform(BuildCookRun))
 			{
-				Owner.Profile->AddCustomError(TEXT("Validation_ZenNeedsRemotePermission"), LOCTEXT("Error_ZenNeedsRemotePermission", "Zen Workflows for remote devices need [Zen.AutoLaunch]:AllowRemoteNetworkService=true. It cannot be overridden in a profile"), BuildCookRun.ToSharedPtr());
+				Owner.Profile->AddCustomError(TEXT("Validation_ZenNeedsRemotePermission"), LOCTEXT("Error_ZenNeedsRemotePermission", "Zen Workflows for remote devices need [Zen.AutoLaunch]:RemoteNetworkService to be a non-None value. It cannot be overridden in a profile"), BuildCookRun.ToSharedPtr());
 			}
 		}
 
@@ -1220,7 +1220,7 @@ namespace ProjectLauncher
 		if ((InContentScheme == EContentScheme::ZenStreaming || InContentScheme == EContentScheme::ZenPakStreaming) 
 			&& !ProjectSettings.bAllowRemoteNetworkService && Model->IsUsingRemotePlatform(BuildCookRun))
 		{
-			OutReason = LOCTEXT("NoZenReason", "Zen Streaming to a remote device requires AllowRemoteNetworkService");
+			OutReason = LOCTEXT("NoZenReason", "Zen Streaming to a remote device requires RemoteNetworkService");
 			return false;
 		}
 
diff --git a/Engine/Source/Developer/IoStoreUtilities/Private/ZenStoreWriter.cpp b/Engine/Source/Developer/IoStoreUtilities/Private/ZenStoreWriter.cpp
--- a/Engine/Source/Developer/IoStoreUtilities/Private/ZenStoreWriter.cpp
+++ b/Engine/Source/Developer/IoStoreUtilities/Private/ZenStoreWriter.cpp
@@ -428,13 +428,15 @@ void FZenStoreWriter::WriteProjectStoreFile()
 	// FZenCookArtifactReader->InitializeConnection.
 	FStringView ServiceHostName = TEXT("localhost");
 	uint16 ServicePort = 8558;
+	FString HostAuthJson;
 	Zen::FZenProjectStoreWriter ProjectStoreWriter(*GetProjectStorePath());
 #if UE_WITH_ZEN
 	const Zen::FZenServiceInstance& ZenServiceInstance = HttpClient->GetZenServiceInstance();
 	ServiceHostName = ZenServiceInstance.GetEndpoint().GetName();
 	ServicePort = ZenServiceInstance.GetEndpoint().GetPort();
+	Zen::GetLocalServiceHostAuth(HostAuthJson);
 #endif
-	ProjectStoreWriter.Write(ServiceHostName, ServicePort, IsLocalConnection, *ProjectId, *OplogId, TargetPlatformFName);
+	ProjectStoreWriter.Write(ServiceHostName, ServicePort, HostAuthJson, IsLocalConnection, *ProjectId, *OplogId, TargetPlatformFName);
 }
 
 void FZenStoreWriter::Initialize(const FCookInfo& Info)
diff --git a/Engine/Source/Developer/Zen/Private/BuildServerInterface.cpp b/Engine/Source/Developer/Zen/Private/BuildServerInterface.cpp
--- a/Engine/Source/Developer/Zen/Private/BuildServerInterface.cpp
+++ b/Engine/Source/Developer/Zen/Private/BuildServerInterface.cpp
@@ -1380,8 +1380,11 @@ FBuildServiceInstance::BuildTransferThreadLoop()
 					FString OplogLifetimeMarkerPath = FPaths::Combine(DestinationDirectory,
 														TEXT("ue.projectstore"));
 
+					FString HostAuthJson;
+					UE::Zen::GetLocalServiceHostAuth(HostAuthJson);
+
 					Zen::FZenProjectStoreWriter ProjectStoreWriter(*OplogLifetimeMarkerPath);
-					ProjectStoreWriter.Write(TEXT("[::1]"), LocalPort, true, DestinationProjectId, DestinationOplogId, State->TargetPlatformName);
+					ProjectStoreWriter.Write(TEXT("[::1]"), LocalPort, HostAuthJson, true, DestinationProjectId, DestinationOplogId, State->TargetPlatformName);
 
 					FPaths::NormalizeDirectoryName(EngineDir);
 					FString ProjectDir = FPaths::GetPath(ProjectFilePath);
diff --git a/Engine/Source/Developer/Zen/Private/ZenProjectStoreWriter.cpp b/Engine/Source/Developer/Zen/Private/ZenProjectStoreWriter.cpp
--- a/Engine/Source/Developer/Zen/Private/ZenProjectStoreWriter.cpp
+++ b/Engine/Source/Developer/Zen/Private/ZenProjectStoreWriter.cpp
@@ -5,10 +5,12 @@
 #if UE_WITH_ZEN
 
 #include "Containers/StringView.h"
+#include "Dom/JsonObject.h"
 #include "Experimental/ZenServerInterface.h"
 #include "HAL/FileManager.h"
 #include "HAL/PlatformMisc.h"
 #include "Policies/PrettyJsonPrintPolicy.h"
+#include "Serialization/JsonReader.h"
 #include "Serialization/JsonSerializer.h"
 #include "Serialization/JsonWriter.h"
 #include "SocketSubsystem.h"
@@ -21,7 +23,7 @@ FZenProjectStoreWriter::FZenProjectStoreWriter(const TCHAR* TargetFileName)
 {
 }
 
-bool FZenProjectStoreWriter::Write(FStringView ServiceHostName, uint16 ServicePort, bool bIsRunningLocally, FStringView ProjectId, FStringView OplogId, FName TargetPlatformName)
+bool FZenProjectStoreWriter::Write(FStringView ServiceHostName, uint16 ServicePort, const FString& HostAuthJson, bool bIsRunningLocally, FStringView ProjectId, FStringView OplogId, FName TargetPlatformName)
 {
 	if (!OplogMarker)
 	{
@@ -84,6 +86,32 @@ bool FZenProjectStoreWriter::Write(FStringView ServiceHostName, uint16 ServicePo
 		}
 	}
 	Writer->WriteValue(TEXT("hostport"), ServicePort);
+	if (!HostAuthJson.IsEmpty())
+	{
+		TSharedPtr<FJsonObject> AuthObject;
+		TSharedRef<TJsonReader<>> Reader = TJsonReaderFactory<>::Create(HostAuthJson);
+		if (FJsonSerializer::Deserialize(Reader, AuthObject) && AuthObject.IsValid())
+		{
+			Writer->WriteObjectStart(TEXT("hostauth"));
+			for (const auto& Pair : AuthObject->Values)
+			{
+				FString Key(Pair.Key);
+				if (Pair.Value->Type == EJson::String)
+				{
+					Writer->WriteValue(Key, Pair.Value->AsString());
+				}
+				else if (Pair.Value->Type == EJson::Boolean)
+				{
+					Writer->WriteValue(Key, Pair.Value->AsBool());
+				}
+				else if (Pair.Value->Type == EJson::Number)
+				{
+					Writer->WriteValue(Key, Pair.Value->AsNumber());
+				}
+			}
+			Writer->WriteObjectEnd();
+		}
+	}
 	Writer->WriteValue(TEXT("projectid"), ProjectId);
 	Writer->WriteValue(TEXT("oplogid"), OplogId);
 	Writer->WriteObjectEnd();
diff --git a/Engine/Source/Developer/Zen/Private/ZenServerInterface.cpp b/Engine/Source/Developer/Zen/Private/ZenServerInterface.cpp
--- a/Engine/Source/Developer/Zen/Private/ZenServerInterface.cpp
+++ b/Engine/Source/Developer/Zen/Private/ZenServerInterface.cpp
@@ -35,6 +35,8 @@
 #include "Serialization/JsonWriter.h"
 #include "String/LexFromString.h"
 #include "Serialization/CompactBinaryWriter.h"
+#include "Misc/Guid.h"
+#include "Policies/CondensedJsonPrintPolicy.h"
 
 #if PLATFORM_WINDOWS
 #	include "Windows/AllowWindowsPlatformTypes.h"
@@ -291,6 +293,154 @@ CleanOutOfDateServicePluginConfigs()
 	});
 }
 
+static FString
+GetSecurityConfigPath(const FString& InstallPath)
+{
+	return FPaths::Combine(InstallPath, TEXT("security-config.json"));
+}
+
+static FString
+BuildSecurityConfigJson(const TSharedPtr<FJsonObject>& AuthObject)
+{
+	FString Json;
+	TSharedRef<TJsonWriter<>> Writer = TJsonWriterFactory<>::Create(&Json);
+	Writer->WriteObjectStart();
+	Writer->WriteObjectStart(TEXT("http"));
+	Writer->WriteObjectStart(TEXT("root"));
+	Writer->WriteObjectStart(TEXT("filter"));
+
+	FString AuthType;
+	if (AuthObject->TryGetStringField(TEXT("type"), AuthType) && AuthType == TEXT("password"))
+	{
+		FString Username;
+		FString Password;
+		AuthObject->TryGetStringField(TEXT("username"), Username);
+		AuthObject->TryGetStringField(TEXT("password"), Password);
+
+		Writer->WriteValue(TEXT("type"), TEXT("password"));
+		Writer->WriteObjectStart(TEXT("config"));
+		Writer->WriteObjectStart(TEXT("basic"));
+		Writer->WriteValue(TEXT("username"), Username);
+		Writer->WriteValue(TEXT("password"), Password);
+		Writer->WriteObjectEnd(); // basic
+		Writer->WriteValue(TEXT("protect-machine-local-requests"), false);
+		Writer->WriteObjectEnd(); // config
+	}
+
+	Writer->WriteObjectEnd(); // filter
+	Writer->WriteObjectEnd(); // root
+	Writer->WriteObjectEnd(); // http
+	Writer->WriteObjectEnd();
+	Writer->Close();
+	return Json;
+}
+
+bool
+GetLocalServiceHostAuth(FString& OutAuthJson)
+{
+	return FPlatformMisc::GetStoredValue(TEXT("Epic Games"), TEXT("Zen"), TEXT("HostAuth"), OutAuthJson)
+		&& !OutAuthJson.IsEmpty();
+}
+
+static void
+SetLocalServiceHostAuth(const FString& AuthJson)
+{
+	FPlatformMisc::SetStoredValue(TEXT("Epic Games"), TEXT("Zen"), TEXT("HostAuth"), AuthJson);
+}
+
+static void
+ClearLocalServiceHostAuth()
+{
+	FPlatformMisc::DeleteStoredValue(TEXT("Epic Games"), TEXT("Zen"), TEXT("HostAuth"));
+}
+
+static bool
+GetGeneratedStaticKey(FString& OutKey)
+{
+	return FPlatformMisc::GetStoredValue(TEXT("Epic Games"), TEXT("Zen"), TEXT("GeneratedStaticKey"), OutKey)
+		&& !OutKey.IsEmpty();
+}
+
+static void
+SetGeneratedStaticKey(const FString& Key)
+{
+	FPlatformMisc::SetStoredValue(TEXT("Epic Games"), TEXT("Zen"), TEXT("GeneratedStaticKey"), Key);
+}
+
+// Callers must hold a system-wide critical section that serializes access to
+// the security config file for the given InstallPath (e.g. ZenServerLaunch
+// for auto-launch, ZenServerInstall for system service).
+static bool
+ConditionalUpdateSecurityConfig(ERemoteNetworkService RemoteNetworkService, const FString& InstallPath)
+{
+	const bool bEnabled = (RemoteNetworkService == ERemoteNetworkService::GeneratedStaticKey);
+	const FString SecurityConfigPath = GetSecurityConfigPath(InstallPath);
+	const bool bFileExists = IFileManager::Get().FileExists(*SecurityConfigPath);
+
+	if (!bEnabled)
+	{
+		// Wipe HostAuth so consumers don't write stale auth data to ue.projectstore.
+		// GeneratedStaticKeyAuth is left intact so switching back to GeneratedStaticKey
+		// mode reuses the same key rather than generating a new one.
+		ClearLocalServiceHostAuth();
+
+		if (bFileExists)
+		{
+			IFileManager::Get().Delete(*SecurityConfigPath, false /*RequireExists*/, false /*EvenReadOnly*/, true /*Quiet*/);
+			return true;
+		}
+		return false;
+	}
+
+	// Retrieve or generate the static key. The key is persisted in GeneratedStaticKey
+	// so it survives mode switches to None/Unsecured and back.
+	FString GeneratedStaticKey;
+	if (!GetGeneratedStaticKey(GeneratedStaticKey))
+	{
+		FString NewKey = FGuid::NewGuid().ToString(EGuidFormats::Digits);
+		SetGeneratedStaticKey(NewKey);
+		GeneratedStaticKey = NewKey;
+	}
+
+	FString NewAuthJson;
+	TSharedRef<TJsonWriter<TCHAR, TCondensedJsonPrintPolicy<TCHAR>>> Writer = TJsonWriterFactory<TCHAR, TCondensedJsonPrintPolicy<TCHAR>>::Create(&NewAuthJson);
+	Writer->WriteObjectStart();
+	Writer->WriteValue(TEXT("type"), TEXT("password"));
+	Writer->WriteValue(TEXT("username"), TEXT("ue"));
+	Writer->WriteValue(TEXT("password"), GeneratedStaticKey);
+	Writer->WriteObjectEnd();
+	Writer->Close();
+
+	FString AuthJson;
+	if (!GetLocalServiceHostAuth(AuthJson) || AuthJson != NewAuthJson)
+	{
+		// Publish to HostAuth so consumers (C#, Python, etc.) can read it.
+		SetLocalServiceHostAuth(NewAuthJson);
+		AuthJson = NewAuthJson;
+	}
+
+	TSharedPtr<FJsonObject> AuthObject;
+	TSharedRef<TJsonReader<>> AuthReader = TJsonReaderFactory<>::Create(AuthJson);
+	if (!FJsonSerializer::Deserialize(AuthReader, AuthObject) || !AuthObject.IsValid())
+	{
+		return false;
+	}
+
+	FString ExpectedJson = BuildSecurityConfigJson(AuthObject);
+	FString ExistingJson;
+	if (FFileHelper::LoadFileToString(ExistingJson, *SecurityConfigPath) && ExistingJson == ExpectedJson)
+	{
+		return false;
+	}
+
+	UE_LOG(LogZenServiceInstance, Display, TEXT("Writing security configuration to '%s'"), *SecurityConfigPath);
+	if (!FFileHelper::SaveStringToFile(ExpectedJson, *SecurityConfigPath, FFileHelper::EEncodingOptions::ForceUTF8WithoutBOM))
+	{
+		UE_LOG(LogZenServiceInstance, Error, TEXT("Failed to write security configuration to '%s'"), *SecurityConfigPath);
+	}
+	return true;
+}
+
 static FString
 GetServiceCopyInstallPath()
 {
@@ -1142,6 +1292,39 @@ static FString ZenGetInstallModeToString(FServiceAutoLaunchSettings::EInstallMod
 	return "";
 }
 
+static ERemoteNetworkService ZenGetRemoteNetworkServiceFromString(const FString& Value)
+{
+	if (Value.IsEmpty() || FCString::Stricmp(*Value, TEXT("None")) == 0)
+	{
+		return ERemoteNetworkService::None;
+	}
+	else if (FCString::Stricmp(*Value, TEXT("Unsecured")) == 0)
+	{
+		return ERemoteNetworkService::Unsecured;
+	}
+	else if (FCString::Stricmp(*Value, TEXT("GeneratedStaticKey")) == 0)
+	{
+		return ERemoteNetworkService::GeneratedStaticKey;
+	}
+	else
+	{
+		checkf(false, TEXT("Unrecognized remote network service type: {%s}"), *Value);
+		return ERemoteNetworkService::None;
+	}
+}
+
+static FString ZenGetRemoteNetworkServiceToString(ERemoteNetworkService Value)
+{
+	switch (Value)
+	{
+	case ERemoteNetworkService::None:          return TEXT("None");
+	case ERemoteNetworkService::Unsecured:     return TEXT("Unsecured");
+	case ERemoteNetworkService::GeneratedStaticKey: return TEXT("GeneratedStaticKey");
+	}
+	checkf(false, TEXT("Unrecognized remote network service type: %d"), static_cast<int>(Value));
+	return TEXT("");
+}
+
 bool FServicePluginSettings::ReadFromConfig(const FString& InPluginName)
 {
 	const FString PluginSectionName = FString::Format(TEXT("Zen.Plugin.{0}"), {*InPluginName}); 
@@ -1262,7 +1445,11 @@ static TOptional<FSystemServiceSettings> ReadSystemServiceSettingsFromConfig()
 
 	GConfig->GetString(SystemServiceConfigSection, TEXT("InstallPath"), SystemServiceSettings.InstallPath, GEngineIni);
 	GConfig->GetString(SystemServiceConfigSection, TEXT("ExtraArgs"), SystemServiceSettings.ExtraArgs, GEngineIni);
-	GConfig->GetBool(SystemServiceConfigSection, TEXT("AllowRemoteNetworkService"), SystemServiceSettings.bAllowRemoteNetworkService, GEngineIni);
+	FString RemoteNetworkServiceStr;
+	if (GConfig->GetString(SystemServiceConfigSection, TEXT("RemoteNetworkService"), RemoteNetworkServiceStr, GEngineIni))
+	{
+		SystemServiceSettings.RemoteNetworkService = ZenGetRemoteNetworkServiceFromString(RemoteNetworkServiceStr);
+	}
 	EnsureEditorSettingsConfigLoaded();
 	GConfig->GetBool(TEXT("/Script/UnrealEd.CrashReportsPrivacySettings"), TEXT("bSendUnattendedBugReports"), SystemServiceSettings.bSendUnattendedBugReports, GEditorSettingsIni);
 
@@ -1368,7 +1555,11 @@ FServiceSettings::ReadFromConfig()
 			GConfig->GetBool(AutoLaunchConfigSection, TEXT("ShowConsole"), AutoLaunchSettings.bShowConsole, GEngineIni);
 			GConfig->GetBool(AutoLaunchConfigSection, TEXT("LimitProcessLifetime"), AutoLaunchSettings.bLimitProcessLifetime, GEngineIni);
 			ApplyProcessLifetimeOverride(AutoLaunchSettings.bLimitProcessLifetime);
-			GConfig->GetBool(AutoLaunchConfigSection, TEXT("AllowRemoteNetworkService"), AutoLaunchSettings.bAllowRemoteNetworkService, GEngineIni);
+			FString RemoteNetworkServiceStr;
+			if (GConfig->GetString(AutoLaunchConfigSection, TEXT("RemoteNetworkService"), RemoteNetworkServiceStr, GEngineIni))
+			{
+				AutoLaunchSettings.RemoteNetworkService = ZenGetRemoteNetworkServiceFromString(RemoteNetworkServiceStr);
+			}
 			FString InstallMode;
 			if (GConfig->GetString(AutoLaunchConfigSection, TEXT("InstallMode"), InstallMode, GEngineIni))
 			{
@@ -1406,7 +1597,7 @@ FServiceSettings::ReadFromCompactBinary(FCbFieldView Field)
 				SystemServiceSettings.ExtraArgs = FString(SystemSettingsObject["ExtraArgs"].AsString());
 				SystemServiceSettings.DesiredPort = SystemSettingsObject["DesiredPort"].AsInt16();
 				SystemServiceSettings.InstallPath = FString(SystemSettingsObject["InstallPath"].AsString());
-				SystemServiceSettings.bAllowRemoteNetworkService = SystemSettingsObject["AllowRemoteNetworkService"].AsBool();
+				SystemServiceSettings.RemoteNetworkService = ZenGetRemoteNetworkServiceFromString(FString(SystemSettingsObject["RemoteNetworkService"].AsString()));
 				SystemServiceSettings.bSendUnattendedBugReports = SystemSettingsObject["SendUnattendedBugReports"].AsBool();
 			}
 		}
@@ -1439,7 +1630,7 @@ FServiceSettings::ReadFromCompactBinary(FCbFieldView Field)
 					AutoLaunchSettings.bIsDefaultDataPath = AutoLaunchSettingsObject["IsDefaultDataPath"].AsBool();
 					AutoLaunchSettings.bLimitProcessLifetime = AutoLaunchSettingsObject["LimitProcessLifetime"].AsBool();
 					ApplyProcessLifetimeOverride(AutoLaunchSettings.bLimitProcessLifetime);
-					AutoLaunchSettings.bAllowRemoteNetworkService = AutoLaunchSettingsObject["AllowRemoteNetworkService"].AsBool();
+					AutoLaunchSettings.RemoteNetworkService = ZenGetRemoteNetworkServiceFromString(FString(AutoLaunchSettingsObject["RemoteNetworkService"].AsString()));
 					AutoLaunchSettings.bSendUnattendedBugReports = AutoLaunchSettingsObject["SendUnattendedBugReports"].AsBool();
 					AutoLaunchSettings.bIsDefaultSharedRunContext = AutoLaunchSettingsObject["IsDefaultSharedRunContext"].AsBool(AutoLaunchSettings.bIsDefaultSharedRunContext);
 					AutoLaunchSettings.InstallMode = ZenGetInstallModeFromString(FString(AutoLaunchSettingsObject["InstallMode"].AsString()));
@@ -1511,7 +1702,7 @@ FServiceSettings::WriteToCompactBinary(FCbWriter& Writer) const
 		Writer << "DesiredPort" << SystemServiceSettings.DesiredPort;
 		Writer << "InstallPath" << SystemServiceSettings.InstallPath;
 		Writer << "SendUnattendedBugReports" << SystemServiceSettings.bSendUnattendedBugReports;
-		Writer << "AllowRemoteNetworkService" << SystemServiceSettings.bAllowRemoteNetworkService;
+		Writer << "RemoteNetworkService" << ZenGetRemoteNetworkServiceToString(SystemServiceSettings.RemoteNetworkService);
 		Writer.EndObject();
 	}
 
@@ -1535,7 +1726,7 @@ FServiceSettings::WriteToCompactBinary(FCbWriter& Writer) const
 		Writer << "ShowConsole" << AutoLaunchSettings.bShowConsole;
 		Writer << "IsDefaultDataPath" << AutoLaunchSettings.bIsDefaultDataPath;
 		Writer << "LimitProcessLifetime" << AutoLaunchSettings.bLimitProcessLifetime;
-		Writer << "AllowRemoteNetworkService" << AutoLaunchSettings.bAllowRemoteNetworkService;
+		Writer << "RemoteNetworkService" << ZenGetRemoteNetworkServiceToString(AutoLaunchSettings.RemoteNetworkService);
 		Writer << "SendUnattendedBugReports" << AutoLaunchSettings.bSendUnattendedBugReports;
 		Writer << "IsDefaultSharedRunContext" << AutoLaunchSettings.bIsDefaultSharedRunContext;
 		Writer << "InstallMode" << ZenGetInstallModeToString(AutoLaunchSettings.InstallMode);
@@ -1836,7 +2027,7 @@ IsZenProcessUsingDataDir(const TCHAR* LockFilePath, ZenLockFileData* OutLockFile
 }
 
 static FString
-DetermineCmdLineWithoutTransientComponents(const FString& DataPath, const FString& ExtraArgs, const TArray<FServicePluginSettings>& Plugins, bool bSendUnattendedBugReports, bool bAllowRemoteNetworkService, int16 OverrideDesiredPort)
+DetermineCmdLineWithoutTransientComponents(const FString& DataPath, const FString& ExtraArgs, const TArray<FServicePluginSettings>& Plugins, bool bSendUnattendedBugReports, ERemoteNetworkService RemoteNetworkService, int16 OverrideDesiredPort, const FString& SecurityConfigPath)
 {
 	FString PlatformDataPath(DataPath);
 	FPaths::MakePlatformFilename(PlatformDataPath);
@@ -1890,9 +2081,22 @@ DetermineCmdLineWithoutTransientComponents(const FString& DataPath, const FStrin
 		Parms.Append(TEXT(" --no-sentry"));
 	}
 
-	if (!bAllowRemoteNetworkService)
+	switch (RemoteNetworkService)
 	{
+	case ERemoteNetworkService::None:
 		Parms.Append(TEXT(" --http-forceloopback"));
+		break;
+	case ERemoteNetworkService::Unsecured:
+		break;
+	case ERemoteNetworkService::GeneratedStaticKey:
+		if (!SecurityConfigPath.IsEmpty())
+		{
+			Parms.Appendf(TEXT(" --security-config-path \"%s\""), *SecurityConfigPath);
+		}
+		break;
+	default:
+		checkf(false, TEXT("Unrecognized remote network service type: %d"), static_cast<int>(RemoteNetworkService));
+		break;
 	}
 
 	return Parms;
@@ -2668,6 +2872,12 @@ FZenServiceInstance::Initialize()
 					}
 					else
 					{
+						bool bSecurityConfigChanged = ConditionalUpdateSecurityConfig(AutoLaunchSettings.RemoteNetworkService, GetLocalInstallPath());
+						if (bSecurityConfigChanged)
+						{
+							ShutDownZenServerProcessLockingDataDir(AutoLaunchSettings.DataPath);
+						}
+
 						uint16 Port = GetEndpoint().GetPort();
 
 						bHasLaunchedLocal = AutoLaunch(AutoLaunchSettings, ExecutablePath, ExecutionContextFilePath, Endpoint);
@@ -3489,12 +3699,16 @@ bool FZenServiceInstance::ConditionalUpdateSystemServiceInstall(const FSystemSer
 		return false;
 	}
 
+	FString SystemInstallPath = GetSystemInstallPath(InSettings);
+	bool bSecurityConfigChanged = ConditionalUpdateSecurityConfig(InSettings.RemoteNetworkService, SystemInstallPath);
+	FString SecurityConfigPath = (InSettings.RemoteNetworkService == ERemoteNetworkService::GeneratedStaticKey) ? GetSecurityConfigPath(SystemInstallPath) : FString();
+
 	FString RunningExecutable;
 	FString RunningCommandLine;
 	uint16 RunningPort;
 
 	TArray<FServicePluginSettings> Plugins;
-	FString DesiredCommandLine = DetermineCmdLineWithoutTransientComponents(InSettings.DataPath, InSettings.ExtraArgs, Plugins, InSettings.bSendUnattendedBugReports, InSettings.bAllowRemoteNetworkService, InSettings.DesiredPort);
+	FString DesiredCommandLine = DetermineCmdLineWithoutTransientComponents(InSettings.DataPath, InSettings.ExtraArgs, Plugins, InSettings.bSendUnattendedBugReports, InSettings.RemoteNetworkService, InSettings.DesiredPort, SecurityConfigPath);
 
 	FZenVersion Version;
 	bool bIsServiceRunning = DetermineSystemServiceInfo(InTreeUtilityPath, RunningExecutable, RunningCommandLine, RunningPort);
@@ -3900,6 +4114,7 @@ FZenServiceInstance::AutoLaunch(const FServiceAutoLaunchSettings& InSettings, co
 	const FString LockFilePath = FPaths::Combine(InSettings.DataPath, TEXT(".lock"));
 
 	FString WorkingDirectory = FPaths::GetPath(ExecutablePath);
+	FString SecurityConfigPath = (InSettings.RemoteNetworkService == ERemoteNetworkService::GeneratedStaticKey) ? GetSecurityConfigPath(GetLocalInstallPath()) : FString();
 
 	ZenLockFileData LockFileState;
 	{
@@ -3937,7 +4152,7 @@ FZenServiceInstance::AutoLaunch(const FServiceAutoLaunchSettings& InSettings, co
 			{
 				FZenLocalServiceRunContext DesiredRunContext;
 				DesiredRunContext.Executable = ExecutablePath;
-				DesiredRunContext.CommandlineArguments = DetermineCmdLineWithoutTransientComponents(InSettings.DataPath, InSettings.ExtraArgs, InSettings.Plugins, InSettings.bSendUnattendedBugReports, InSettings.bAllowRemoteNetworkService, InSettings.DesiredPort);
+				DesiredRunContext.CommandlineArguments = DetermineCmdLineWithoutTransientComponents(InSettings.DataPath, InSettings.ExtraArgs, InSettings.Plugins, InSettings.bSendUnattendedBugReports, InSettings.RemoteNetworkService, InSettings.DesiredPort, SecurityConfigPath);
 				DesiredRunContext.WorkingDirectory = WorkingDirectory;
 				DesiredRunContext.DataPath = InSettings.DataPath;
 				DesiredRunContext.bShowConsole = InSettings.bShowConsole;
@@ -4045,7 +4260,7 @@ FZenServiceInstance::AutoLaunch(const FServiceAutoLaunchSettings& InSettings, co
 			}
 		}
 
-		FString ParmsWithoutTransients = DetermineCmdLineWithoutTransientComponents(InSettings.DataPath, InSettings.ExtraArgs, InSettings.Plugins, InSettings.bSendUnattendedBugReports, InSettings.bAllowRemoteNetworkService, InSettings.DesiredPort);
+		FString ParmsWithoutTransients = DetermineCmdLineWithoutTransientComponents(InSettings.DataPath, InSettings.ExtraArgs, InSettings.Plugins, InSettings.bSendUnattendedBugReports, InSettings.RemoteNetworkService, InSettings.DesiredPort, SecurityConfigPath);
 
 		if (InSettings.InstallMode == FServiceAutoLaunchSettings::EInstallMode::Link)
 		{
diff --git a/Engine/Source/Developer/Zen/Public/Experimental/ZenProjectStoreWriter.h b/Engine/Source/Developer/Zen/Public/Experimental/ZenProjectStoreWriter.h
--- a/Engine/Source/Developer/Zen/Public/Experimental/ZenProjectStoreWriter.h
+++ b/Engine/Source/Developer/Zen/Public/Experimental/ZenProjectStoreWriter.h
@@ -20,7 +20,7 @@ public:
 	UE_API FZenProjectStoreWriter(const TCHAR* TargetFileName);
 	~FZenProjectStoreWriter() = default;
 
-	UE_API bool Write(FStringView ServiceHostName, uint16 ServicePort, bool bIsRunningLocally, FStringView ProjectId, FStringView OplogId, FName TargetPlatformName);
+	UE_API bool Write(FStringView ServiceHostName, uint16 ServicePort, const FString& HostAuthJson, bool bIsRunningLocally, FStringView ProjectId, FStringView OplogId, FName TargetPlatformName);
 private:
 	TUniquePtr<FArchive> OplogMarker;
 };
diff --git a/Engine/Source/Developer/Zen/Public/Experimental/ZenServerInterface.h b/Engine/Source/Developer/Zen/Public/Experimental/ZenServerInterface.h
--- a/Engine/Source/Developer/Zen/Public/Experimental/ZenServerInterface.h
+++ b/Engine/Source/Developer/Zen/Public/Experimental/ZenServerInterface.h
@@ -30,6 +30,13 @@ class FCbFieldView;
 namespace UE::Zen
 {
 
+enum class ERemoteNetworkService : uint8
+{
+	None,
+	Unsecured,
+	GeneratedStaticKey,
+};
+
 struct FServiceConnectSettings
 {
 	FString HostName;
@@ -57,7 +64,7 @@ struct FServiceAutoLaunchSettings
 	bool bShowConsole = false;
 	bool bIsDefaultDataPath = false;
 	bool bLimitProcessLifetime = false;
-	bool bAllowRemoteNetworkService = false;
+	ERemoteNetworkService RemoteNetworkService = ERemoteNetworkService::None;
 	bool bSendUnattendedBugReports = false;
 	bool bIsDefaultSharedRunContext = true;
 	enum class EInstallMode
@@ -74,7 +81,7 @@ struct FSystemServiceSettings
 	FString ExtraArgs;
 	uint16 DesiredPort = 8558;
 	bool bSendUnattendedBugReports = false;
-	bool bAllowRemoteNetworkService = false;
+	ERemoteNetworkService RemoteNetworkService = ERemoteNetworkService::None;
 };
 
 struct FServiceSettings
@@ -87,6 +94,12 @@ struct FServiceSettings
 	inline bool IsAutoLaunch() const { return IsSystemService() || SettingsVariant.IsType<FServiceAutoLaunchSettings>(); }
 	inline bool IsConnectExisting() const { return !IsSystemService() && SettingsVariant.IsType<FServiceConnectSettings>(); }
 	inline bool IsRequirePlugins() const { return IsAutoLaunch() && !SettingsVariant.Get<FServiceAutoLaunchSettings>().Plugins.IsEmpty(); }
+	inline ERemoteNetworkService GetRemoteNetworkService() const
+	{
+		if (bSystemService) return SystemServiceSettings.RemoteNetworkService;
+		if (IsAutoLaunch()) return SettingsVariant.Get<FServiceAutoLaunchSettings>().RemoteNetworkService;
+		return ERemoteNetworkService::None;
+	}
 
 	UE_API bool ReadFromConfig();
 	UE_API bool ReadFromCompactBinary(FCbFieldView Field);
@@ -229,6 +242,8 @@ UE_API FString GetLocalServiceInstallPath();
 UE_API FString GetLocalServiceInstallVersion(bool bDetailed = true);
 UE_API FString GetLocalInstallUtilityPath();
 
+UE_API bool GetLocalServiceHostAuth(FString& OutAuthJson);
+
 struct FGCStatus
 {
 	FString Description;
diff --git a/Engine/Source/Programs/Shared/EpicGames.ProjectStore/ProjectStoreTypes.cs b/Engine/Source/Programs/Shared/EpicGames.ProjectStore/ProjectStoreTypes.cs
--- a/Engine/Source/Programs/Shared/EpicGames.ProjectStore/ProjectStoreTypes.cs
+++ b/Engine/Source/Programs/Shared/EpicGames.ProjectStore/ProjectStoreTypes.cs
@@ -8,6 +8,7 @@ using System.Linq;
 using System.Net.Http.Headers;
 using System.Net.Http;
 using System.Net;
+using System.Runtime.InteropServices;
 using System.Text.Json.Serialization;
 using System.Text.Json;
 using System.Text;
@@ -60,6 +61,11 @@ namespace EpicGames.ProjectStore
 		/// </summary>
 		public int HostPort { get; set; } = 8558;
 		/// <summary>
+		/// Authentication details for connecting to zenserver. Null if no authentication is required.
+		/// Stored as an opaque JSON object that is passed through from the platform storage.
+		/// </summary>
+		public JsonElement? HostAuth { get; set; }
+		/// <summary>
 		/// The project identifier for the stored data in zenserver
 		/// </summary>
 		public string? ProjectId { get; set; }
@@ -252,6 +258,123 @@ namespace EpicGames.ProjectStore
 			return new FileReference(Path.Combine(engineDirectory.FullName, "Binaries", platform, String.Format("zen{0}", RuntimePlatform.ExeExtension)));
 		}
 
+		/// <summary>
+		/// Retrieves the Zen host auth data from the platform credential store as an opaque JSON element,
+		/// matching the storage used by FPlatformMisc::GetStoredValue(TEXT("Epic Games"), TEXT("Zen"), TEXT("HostAuth")).
+		/// Returns null if no auth data is stored or if reading/parsing fails.
+		/// </summary>
+		private static JsonElement? GetZenHostAuth()
+		{
+			try
+			{
+				string? authJson = null;
+				if (RuntimeInformation.IsOSPlatform(OSPlatform.Windows))
+				{
+					using Microsoft.Win32.RegistryKey? key = Microsoft.Win32.Registry.CurrentUser.OpenSubKey(@"Software\Epic Games\Zen");
+					authJson = key?.GetValue("HostAuth") as string;
+				}
+				else
+				{
+					// On Mac and Linux, FPlatformMisc::SetStoredValue writes to an INI file:
+					//   ApplicationSettingsDir() / "Epic Games" / "KeyValueStore.ini"
+					// with section [Zen] and key HostAuth=<value>
+					// Mac ApplicationSettingsDir: ~/Library/Application Support/Epic/
+					// Linux ApplicationSettingsDir: ~/.config/Epic/
+					string appSettingsDir;
+					if (RuntimeInformation.IsOSPlatform(OSPlatform.OSX))
+					{
+						appSettingsDir = Path.Combine(Environment.GetFolderPath(Environment.SpecialFolder.UserProfile),
+							"Library", "Application Support", "Epic");
+					}
+					else
+					{
+						appSettingsDir = Path.Combine(Environment.GetFolderPath(Environment.SpecialFolder.UserProfile),
+							".config", "Epic");
+					}
+
+					string iniPath = Path.Combine(appSettingsDir, "Epic Games", "KeyValueStore.ini");
+					if (File.Exists(iniPath))
+					{
+						authJson = ReadIniValue(iniPath, "Zen", "HostAuth");
+					}
+				}
+
+				if (!String.IsNullOrEmpty(authJson))
+				{
+					using JsonDocument doc = JsonDocument.Parse(authJson);
+					return doc.RootElement.Clone();
+				}
+			}
+			catch
+			{
+				// Silently fail - auth is optional
+			}
+			return null;
+		}
+
+		/// <summary>
+		/// Reads a value from a UE-style INI file given a section name and key name.
+		/// Handles UE's quoted-string format: if the value is wrapped in double quotes,
+		/// the quotes are stripped and escape sequences are processed to match the
+		/// behavior of FParse::QuotedString / FConfigFile::ShouldExportQuotedString.
+		/// </summary>
+		private static string? ReadIniValue(string iniPath, string sectionName, string keyName)
+		{
+			bool inSection = false;
+			string sectionHeader = $"[{sectionName}]";
+			foreach (string rawLine in File.ReadLines(iniPath))
+			{
+				string line = rawLine.Trim();
+				if (line.StartsWith('['))
+				{
+					inSection = String.Equals(line, sectionHeader, StringComparison.OrdinalIgnoreCase);
+					continue;
+				}
+				if (inSection)
+				{
+					int eqIndex = line.IndexOf('=', StringComparison.Ordinal);
+					if (eqIndex > 0 && String.Equals(line.Substring(0, eqIndex), keyName, StringComparison.OrdinalIgnoreCase))
+					{
+						string raw = line.Substring(eqIndex + 1);
+						if (raw.Length >= 2 && raw[0] == '"' && raw[^1] == '"')
+						{
+							return UnescapeIniQuotedString(raw.AsSpan(1, raw.Length - 2));
+						}
+						return raw;
+					}
+				}
+			}
+			return null;
+		}
+
+		/// <summary>
+		/// Unescapes a UE INI quoted string value, matching FParse::QuotedString.
+		/// </summary>
+		private static string UnescapeIniQuotedString(ReadOnlySpan<char> value)
+		{
+			StringBuilder sb = new(value.Length);
+			for (int i = 0; i < value.Length; i++)
+			{
+				if (value[i] != '\\' || i + 1 >= value.Length)
+				{
+					sb.Append(value[i]);
+					continue;
+				}
+				char esc = value[++i];
+				switch (esc)
+				{
+					case '\\': sb.Append('\\'); break;
+					case '"':  sb.Append('"');  break;
+					case '\'': sb.Append('\''); break;
+					case 'n':  sb.Append('\n'); break;
+					case 'r':  sb.Append('\r'); break;
+					case 't':  sb.Append('\t'); break;
+					default:   sb.Append('\\'); sb.Append(esc); break;
+				}
+			}
+			return sb.ToString();
+		}
+
 		/// <summary>
 		/// Writes a ue.projectstore file for the given project and platform to the platform cooked data directory.
 		/// </summary>
@@ -272,6 +395,7 @@ namespace EpicGames.ProjectStore
 			data.ZenServer.IsLocalHost = true;
 			data.ZenServer.HostName = "[::1]";
 			data.ZenServer.HostPort = 8558;
+			data.ZenServer.HostAuth = GetZenHostAuth();
 			data.ZenServer.ProjectId = projectId;
 			data.ZenServer.OplogId = platform;
 			data.ZenServer.RemoteHostNames.AddRange(GetHostAddresses());
diff --git a/Engine/Source/Runtime/StorageServerClient/Private/BuiltInHttpClient/BuiltInHttpClient.cpp b/Engine/Source/Runtime/StorageServerClient/Private/BuiltInHttpClient/BuiltInHttpClient.cpp
--- a/Engine/Source/Runtime/StorageServerClient/Private/BuiltInHttpClient/BuiltInHttpClient.cpp
+++ b/Engine/Source/Runtime/StorageServerClient/Private/BuiltInHttpClient/BuiltInHttpClient.cpp
@@ -46,6 +46,11 @@ FBuiltInHttpClient::FBuiltInHttpClient(TUniquePtr<IBuiltInHttpClientSocketPool>
 {
 }
 
+void FBuiltInHttpClient::SetPersistentHeaders(TArrayView<const FAnsiString> InHeaders)
+{
+	PersistentHeaders = InHeaders;
+}
+
 IStorageServerHttpClient::FResult FBuiltInHttpClient::RequestSync(
 	FAnsiStringView Url,
 	EStorageServerContentType Accept,
@@ -80,6 +85,10 @@ IStorageServerHttpClient::FResult FBuiltInHttpClient::RequestSync(
 		HeaderBuffer
 			<< "Content-Type: " << GetMimeTypeString(PayloadContentType) << "\r\n";
 	}
+	for (const FAnsiString& Header : PersistentHeaders)
+	{
+		HeaderBuffer << Header << "\r\n";
+	}
 	HeaderBuffer << "\r\n";
 
 	IBuiltInHttpClientSocket* Socket = nullptr;
diff --git a/Engine/Source/Runtime/StorageServerClient/Private/BuiltInHttpClient/BuiltInHttpClient.h b/Engine/Source/Runtime/StorageServerClient/Private/BuiltInHttpClient/BuiltInHttpClient.h
--- a/Engine/Source/Runtime/StorageServerClient/Private/BuiltInHttpClient/BuiltInHttpClient.h
+++ b/Engine/Source/Runtime/StorageServerClient/Private/BuiltInHttpClient/BuiltInHttpClient.h
@@ -33,6 +33,8 @@ public:
 	FBuiltInHttpClient(TUniquePtr<IBuiltInHttpClientSocketPool> InSocketPool, FString InHostname);
 	virtual ~FBuiltInHttpClient() override = default;
 
+	virtual void SetPersistentHeaders(TArrayView<const FAnsiString> InHeaders) override;
+
 	virtual FResult RequestSync(
 		FAnsiStringView Url,
 		EStorageServerContentType Accept,
@@ -43,7 +45,7 @@ public:
 		float TimeoutSeconds,
 		const bool bReportErrors
 	) override;
-	
+
 	virtual void RequestAsync(
 		FResultCallback&& Callback,
 		FAnsiStringView Url,
@@ -62,6 +64,7 @@ public:
 private:
 	TUniquePtr<IBuiltInHttpClientSocketPool> SocketPool;
 	FString Hostname;
+	TArray<FAnsiString> PersistentHeaders;
 };
 
 #endif
diff --git a/Engine/Source/Runtime/StorageServerClient/Private/StorageServerConnection.cpp b/Engine/Source/Runtime/StorageServerClient/Private/StorageServerConnection.cpp
--- a/Engine/Source/Runtime/StorageServerClient/Private/StorageServerConnection.cpp
+++ b/Engine/Source/Runtime/StorageServerClient/Private/StorageServerConnection.cpp
@@ -99,9 +99,11 @@ static FCbObject ProcessResponse(FIoBuffer IoBuffer, EStorageServerContentType C
 	return ResponseObj;
 }
 
-bool FStorageServerConnection::Initialize(TArrayView<const FString> HostAddresses, const int32 Port, const FAnsiStringView& InBaseURI)
+bool FStorageServerConnection::Initialize(TArrayView<const FString> HostAddresses, const int32 Port, const FAnsiStringView& InBaseURI, TArrayView<const FAnsiString> InPersistentHeaders)
 {
 	BaseURI = InBaseURI;
+	PersistentHeaders = InPersistentHeaders;
+
 	TArray<FString> SortedHostAddresses = SortHostAddressesByLocalSubnet(HostAddresses, Port);
 	if (BaseURI.StartsWith("/ws/"))
 	{
@@ -111,6 +113,7 @@ bool FStorageServerConnection::Initialize(TArrayView<const FString> HostAddresse
 	for (const FString& HostAddress : SortedHostAddresses)
 	{
 		HttpClient = CreateHttpClient(HostAddress, Port);
+		HttpClient->SetPersistentHeaders(PersistentHeaders);
 		CurrentHostAddr = HostAddress;
 		if (HandshakeRequest())
 		{
diff --git a/Engine/Source/Runtime/StorageServerClient/Private/StorageServerConnection.h b/Engine/Source/Runtime/StorageServerClient/Private/StorageServerConnection.h
--- a/Engine/Source/Runtime/StorageServerClient/Private/StorageServerConnection.h
+++ b/Engine/Source/Runtime/StorageServerClient/Private/StorageServerConnection.h
@@ -49,7 +49,7 @@ public:
 	FStorageServerConnection() = default;
 	~FStorageServerConnection() = default;
 
-	bool Initialize(TArrayView<const FString> HostAddresses, const int32 Port, const FAnsiStringView& InBaseURI);
+	bool Initialize(TArrayView<const FString> HostAddresses, const int32 Port, const FAnsiStringView& InBaseURI, TArrayView<const FAnsiString> InPersistentHeaders = {});
 
 	struct Workspaces
 	{
@@ -138,6 +138,7 @@ private:
 	TUniquePtr<StorageServer::ICacheStrategy> CacheStrategy;
 	FAnsiString BaseURI;
 	FString CurrentHostAddr;
+	TArray<FAnsiString> PersistentHeaders;
 	bool bIsUsingZenWorkspace = false;		    // is the connection to the /ws/ endpoint
 
 	// Stats
diff --git a/Engine/Source/Runtime/StorageServerClient/Private/StorageServerHttpClient.h b/Engine/Source/Runtime/StorageServerClient/Private/StorageServerHttpClient.h
--- a/Engine/Source/Runtime/StorageServerClient/Private/StorageServerHttpClient.h
+++ b/Engine/Source/Runtime/StorageServerClient/Private/StorageServerHttpClient.h
@@ -17,7 +17,7 @@ enum class EStorageServerContentType : uint8
 	CompressedBinary,
 };
 
-class IStorageServerHttpClient 
+class IStorageServerHttpClient
 {
 public:
 	using FResult = TTuple<TIoStatusOr<FIoBuffer>, EStorageServerContentType>;
@@ -25,6 +25,9 @@ public:
 
 	virtual ~IStorageServerHttpClient() = default;
 
+	/** Set headers that will be included with every request. Each entry should be a complete header line (e.g. "Authorization: Basic dWU6cGFzcw=="). */
+	virtual void SetPersistentHeaders(TArrayView<const FAnsiString> InHeaders) = 0;
+
 	virtual FResult RequestSync(
 		FAnsiStringView Url,
 		EStorageServerContentType Accept = EStorageServerContentType::Unknown,
diff --git a/Engine/Source/Runtime/StorageServerClient/Private/StorageServerPlatformFile.cpp b/Engine/Source/Runtime/StorageServerClient/Private/StorageServerPlatformFile.cpp
--- a/Engine/Source/Runtime/StorageServerClient/Private/StorageServerPlatformFile.cpp
+++ b/Engine/Source/Runtime/StorageServerClient/Private/StorageServerPlatformFile.cpp
@@ -10,6 +10,7 @@
 #include "HAL/RunnableThread.h"
 #include "HAL/CriticalSection.h"
 #include "Misc/App.h"
+#include "Misc/Base64.h"
 #include "Misc/AssertionMacros.h"
 #include "Misc/CommandLine.h"
 #include "Misc/ConfigCacheIni.h"
@@ -1299,6 +1300,28 @@ bool FStorageServerPlatformFile::Initialize(IPlatformFile* Inner, const TCHAR* C
 					{
 						WorkspaceSharePath.Empty();
 					}
+
+					const TSharedPtr<FJsonObject>* HostAuthObjectPtr = nullptr;
+					if (ZenServerObject->TryGetObjectField(TEXT("hostauth"), HostAuthObjectPtr) && HostAuthObjectPtr != nullptr)
+					{
+						const TSharedPtr<FJsonObject>& HostAuthObject = *HostAuthObjectPtr;
+						FString AuthType;
+						if (HostAuthObject->TryGetStringField(TEXT("type"), AuthType))
+						{
+							if (AuthType == TEXT("password"))
+							{
+								FString Username;
+								FString Password;
+								HostAuthObject->TryGetStringField(TEXT("username"), Username);
+								HostAuthObject->TryGetStringField(TEXT("password"), Password);
+								FString Credentials = FString::Printf(TEXT("%s:%s"), *Username, *Password);
+								FString Encoded = FBase64::Encode(Credentials);
+								AuthHeaders.Add(FAnsiString::Printf("Authorization: Basic %s", TCHAR_TO_ANSI(*Encoded)));
+								UE_LOG(LogStorageServerPlatformFile, Display, TEXT("Using password authentication from ue.projectstore"));
+							}
+						}
+					}
+
 					UE_LOG(LogStorageServerPlatformFile, Display, TEXT("Using settings from ue.projectstore: ServerProject='%s' and ServerPlatform='%s'"), *ServerProject, *ServerPlatform);
 				}
 
@@ -1376,7 +1399,7 @@ void FStorageServerPlatformFile::InitializeConnection()
 	}
 #endif
 	Connection.Reset(new FStorageServerConnection());
-	if (Connection->Initialize(HostAddrs, HostPort, MakeBaseURI()))
+	if (Connection->Initialize(HostAddrs, HostPort, MakeBaseURI(), AuthHeaders))
 	{
 #if WITH_STORAGE_SERVER_STARTUP_FILE_CACHE
 		GStorageServerEngineStartupPrecache.Reset( new FStorageServerEngineStartupPrecache(*Connection.Get()) );
diff --git a/Engine/Source/Runtime/StorageServerClient/Private/StorageServerPlatformFile.h b/Engine/Source/Runtime/StorageServerClient/Private/StorageServerPlatformFile.h
--- a/Engine/Source/Runtime/StorageServerClient/Private/StorageServerPlatformFile.h
+++ b/Engine/Source/Runtime/StorageServerClient/Private/StorageServerPlatformFile.h
@@ -183,6 +183,7 @@ private:
 	TDirectoryTree<FString> RemapDirectoriesTree;
 	mutable TArray<FString> HostAddrs;
 	mutable uint16 HostPort = 8558;
+	TArray<FAnsiString> AuthHeaders;
 	bool bAllowPackageIo = true;
 	bool bAbortOnConnectionFailure = true;
 };
