BuildStorageTool: Ensure graceful Ctrl+Break/SIGINT shutdown for zen.exe build transfers instead of immediately doing hard termination.  Tested on Windows and Mac for platform specific behavior.

#Synced-CL 52199356

diff --git a/Engine/Source/Developer/Zen/Private/BuildServerInterface.cpp b/Engine/Source/Developer/Zen/Private/BuildServerInterface.cpp
--- a/Engine/Source/Developer/Zen/Private/BuildServerInterface.cpp
+++ b/Engine/Source/Developer/Zen/Private/BuildServerInterface.cpp
@@ -36,6 +36,15 @@
 #include "Tasks/Task.h"
 #include "StudioTelemetry.h"
 
+#if PLATFORM_WINDOWS
+#	include "Windows/AllowWindowsPlatformTypes.h"
+#	include <Windows.h>
+#	include "Windows/HideWindowsPlatformTypes.h"
+#elif PLATFORM_MAC || PLATFORM_LINUX
+#	include <signal.h>
+#	include <errno.h>
+#endif
+
 #if WITH_SSL
 #include "Ssl.h"
 #endif
@@ -47,6 +56,61 @@ namespace UE::Zen::Build
 
 DEFINE_LOG_CATEGORY_STATIC(LogBuildServiceInstance, Log, All);
 
+constexpr double GracefulTerminationTimeoutSeconds = 15.0;
+
+bool SendSoftTerminationSignal(FProcHandle ProcessHandle)
+{
+#if PLATFORM_WINDOWS
+	// FPlatformProcess::CreateProc always sets CREATE_NEW_PROCESS_GROUP, which
+	// causes Windows to implicitly disable CTRL_C_EVENT for the child. We use
+	// CTRL_BREAK_EVENT instead, which cannot be disabled and can target a
+	// specific process group (the child's PID equals its process group ID).
+	//
+	// GenerateConsoleCtrlEvent requires the caller to share a console with the
+	// target. The child is launched with bLaunchDetached=false, so it inherits
+	// our console when we have one. If we don't (e.g. GUI host), AttachConsole
+	// lets us temporarily join the child's console. We avoid calling FreeConsole
+	// up front to prevent destroying our own console — that would invalidate
+	// stdout/stderr handles and break concurrent I/O on other threads.
+	DWORD ChildPid = ::GetProcessId(ProcessHandle.Get());
+	if (ChildPid == 0)
+	{
+		return false;
+	}
+
+	// Try to attach to the child's console. If we already share one,
+	// AttachConsole fails with ERROR_ACCESS_DENIED and we can proceed directly.
+	bool bDidAttach = !!::AttachConsole(ChildPid);
+	if (!bDidAttach && ::GetLastError() != ERROR_ACCESS_DENIED)
+	{
+		// Child has no console — cannot send a console control event.
+		return false;
+	}
+
+	// CTRL_BREAK_EVENT with a non-zero process group ID only targets that
+	// group, so our own process will not receive the signal.
+	BOOL bResult = ::GenerateConsoleCtrlEvent(CTRL_BREAK_EVENT, ChildPid);
+
+	if (bDidAttach)
+	{
+		::FreeConsole();
+	}
+
+	return bResult != 0;
+
+#elif PLATFORM_MAC || PLATFORM_LINUX
+	pid_t Pid = ProcessHandle.Get();
+	if (Pid <= 0)
+	{
+		return false;
+	}
+	return kill(Pid, SIGINT) == 0;
+
+#else
+	return false;
+#endif
+}
+
 namespace Private
 {
 
@@ -1528,38 +1592,71 @@ FBuildServiceInstance::BuildTransferThreadLoop()
 		const uint64 StartTime = FPlatformTime::Cycles64();
 		while (MonitoredUtilityProcess.Update())
 		{
-			if (bBuildTransferThreadStopping.load(std::memory_order_relaxed))
-			{
-				return;
-			}
 			double Duration = FPlatformTime::ToSeconds64(FPlatformTime::Cycles64() - StartTime);
 			if (bBuildTransferThreadStopping.load(std::memory_order_relaxed) || State->bCancelRequested.load(std::memory_order_relaxed))
 			{
-				MonitoredUtilityProcess.Cancel(true);
-				UE_LOG(LogBuildServiceInstance, Display, TEXT("Canceled launch of zen utility for downloading build data: '%s'."), *ZenUtilityPath);
+				// Send Ctrl+Break / SIGINT to allow graceful shutdown
+				bool bSentSignal = SendSoftTerminationSignal(MonitoredUtilityProcess.GetProcessHandle());
+				if (bSentSignal)
+				{
+					UE_LOG(LogBuildServiceInstance, Display,
+						TEXT("Sent Ctrl+Break to zen utility for build data transfer: '%s'. Waiting up to %.0f seconds for graceful shutdown."),
+						*ZenUtilityPath, GracefulTerminationTimeoutSeconds);
+				}
+				else
+				{
+					UE_LOG(LogBuildServiceInstance, Warning,
+						TEXT("Failed to send Ctrl+Break to zen utility: '%s'. Falling back to hard termination."),
+						*ZenUtilityPath);
+					MonitoredUtilityProcess.Cancel(true);
+				}
 
-				// Wait for execution to be terminated
+				// Wait for graceful exit (or immediate exit if hard-terminated above)
+				const uint64 GraceStartTime = FPlatformTime::Cycles64();
+				bool bExitedGracefully = false;
 				while (MonitoredUtilityProcess.Update())
 				{
-					if (bBuildTransferThreadStopping.load(std::memory_order_relaxed))
-					{
-						return;
-					}
-					Duration = FPlatformTime::ToSeconds64(FPlatformTime::Cycles64() - StartTime);
-					if (Duration > 15.0)
+					double GraceDuration = FPlatformTime::ToSeconds64(FPlatformTime::Cycles64() - GraceStartTime);
+					if (GraceDuration > GracefulTerminationTimeoutSeconds)
 					{
-						UE_LOG(LogBuildServiceInstance, Warning, TEXT("Canceled launch of zen utility for downloading build data: '%s'. Failed waiting for termination."), *ZenUtilityPath);
 						break;
 					}
 					FPlatformProcess::Sleep(0.2f);
 				}
 
+				if (!MonitoredUtilityProcess.Update())
+				{
+					bExitedGracefully = true;
+				}
+
+				// If still running after grace period, hard terminate
+				if (!bExitedGracefully && bSentSignal)
+				{
+					UE_LOG(LogBuildServiceInstance, Warning,
+						TEXT("Zen utility '%s' did not respond to Ctrl+Break signal within %.0f second grace period. Hard-terminating."),
+						*ZenUtilityPath, GracefulTerminationTimeoutSeconds);
+					MonitoredUtilityProcess.Cancel(true);
+
+					const uint64 HardKillStartTime = FPlatformTime::Cycles64();
+					while (MonitoredUtilityProcess.Update())
+					{
+						double HardKillDuration = FPlatformTime::ToSeconds64(FPlatformTime::Cycles64() - HardKillStartTime);
+						if (HardKillDuration > 10.0)
+						{
+							UE_LOG(LogBuildServiceInstance, Warning,
+								TEXT("Zen utility '%s' did not terminate after hard kill within 10 seconds."), *ZenUtilityPath);
+							break;
+						}
+						FPlatformProcess::Sleep(0.2f);
+					}
+				}
+
 				FWriteScopeLock _(State->Lock);
 				State->TransferEndTime = FMonotonicTimePoint::Now();
 				State->Status = EBuildTransferStatus::Canceled;
 				double DurationSeconds = (State->TransferEndTime - State->TransferStartTime).ToSeconds();
-				UE_LOG(LogBuildServiceInstance, Display, TEXT("Transfer canceled after %.1f minutes."), DurationSeconds/60.0f);
-				State->OutputLogFile->Logf(TEXT("Transfer canceled after %.1f minutes."), DurationSeconds/60.0f);
+				UE_LOG(LogBuildServiceInstance, Display, TEXT("Transfer canceled after %.1f minutes."), DurationSeconds / 60.0f);
+				State->OutputLogFile->Logf(TEXT("Transfer canceled after %.1f minutes."), DurationSeconds / 60.0f);
 				State->OutputLogFile.Reset();
 				RecordTransferStateEvent(State);
 
