Remove use of O_CLOEXEC from shm_open as it is rejected in MacOS 26.4 and later.  It is instead acceptable to call fcntl later with O_CLOEXEC.

#Synced-CL 52547800

diff --git a/Engine/Source/Developer/Zen/Private/ZenServerState.cpp b/Engine/Source/Developer/Zen/Private/ZenServerState.cpp
--- a/Engine/Source/Developer/Zen/Private/ZenServerState.cpp
+++ b/Engine/Source/Developer/Zen/Private/ZenServerState.cpp
@@ -263,12 +263,13 @@ ZenServerState::ZenServerState(bool ReadOnly)
 		return;
 	}
 #elif PLATFORM_UNIX || PLATFORM_MAC
-	int OFlag = ReadOnly ? (O_RDONLY | O_CLOEXEC) : (O_RDWR | O_CREAT | O_CLOEXEC);
+	int OFlag = ReadOnly ? (O_RDONLY) : (O_RDWR | O_CREAT);
 	int Fd = shm_open("/UnrealEngineZen", OFlag, 0666);
 	if (Fd < 0)
 	{
 		return;
 	}
+	fcntl(Fd, F_SETFD, FD_CLOEXEC);
 	void* hMap = (void*)intptr_t(Fd);
 
 	int Prot = ReadOnly ? PROT_READ : (PROT_WRITE | PROT_READ);
