Remove use of O_CLOEXEC from shared memory descriptors altogether as it is rejected from shm_open calls in MacOS 26.4 and later, and calling it from fcntl is un-necessary according to documentation.

#Synced-CL 52656144

diff --git a/Engine/Source/Developer/Zen/Private/ZenServerState.cpp b/Engine/Source/Developer/Zen/Private/ZenServerState.cpp
--- a/Engine/Source/Developer/Zen/Private/ZenServerState.cpp
+++ b/Engine/Source/Developer/Zen/Private/ZenServerState.cpp
@@ -269,7 +269,6 @@ ZenServerState::ZenServerState(bool ReadOnly)
 	{
 		return;
 	}
-	fcntl(Fd, F_SETFD, FD_CLOEXEC);
 	void* hMap = (void*)intptr_t(Fd);
 
 	int Prot = ReadOnly ? PROT_READ : (PROT_WRITE | PROT_READ);
