Enable running Zen in daemon mode on Mac.

#Synced-CL 52566180

diff --git a/Engine/Source/Developer/Zen/Private/ZenServerInterface.cpp b/Engine/Source/Developer/Zen/Private/ZenServerInterface.cpp
--- a/Engine/Source/Developer/Zen/Private/ZenServerInterface.cpp
+++ b/Engine/Source/Developer/Zen/Private/ZenServerInterface.cpp
@@ -684,6 +684,10 @@ static bool RunZenUtilityElevated(const FString& UtilityPath, const FString& Com
 	ElevatedCommandLine = UtilityPath + TEXT(" ") + CommandLine;
 	ElevatedUtilityPath = TEXT("pkexec");
 #endif
+#if PLATFORM_MAC
+	ElevatedUtilityPath = UtilityPath;
+	ElevatedCommandLine = CommandLine;
+#endif
 
 	if (PassthroughArgs != TEXT(""))
 	{
